#UE4
- Fixed profiler
  - Increased stack size used to generate the profiler event graph
  - Fixed a problem with tracking default stats when the metadata is not ready

diff --git a/Engine/Source/Developer/Profiler/Private/ProfilerDataSource.h b/Engine/Source/Developer/Profiler/Private/ProfilerDataSource.h
--- a/Engine/Source/Developer/Profiler/Private/ProfilerDataSource.h
+++ b/Engine/Source/Developer/Profiler/Private/ProfilerDataSource.h
@@ -967,7 +967,7 @@ class FEventGraphSample : public TSharedFromThis<FEventGraphSample>
 	enum
 	{
 		/** Maximum number of stack when traversing event graph - should be enough to store root nodes and all its children.*/
-		MaxStackSize = 16384
+		MaxStackSize = 65536
 	};
 
 	struct FDuplicateHierarchyTag {};
diff --git a/Engine/Source/Developer/Profiler/Private/ProfilerManager.cpp b/Engine/Source/Developer/Profiler/Private/ProfilerManager.cpp
--- a/Engine/Source/Developer/Profiler/Private/ProfilerManager.cpp
+++ b/Engine/Source/Developer/Profiler/Private/ProfilerManager.cpp
@@ -310,7 +310,7 @@ void FProfilerManager::ProfilerClient_OnMetaDataUpdated( const FGuid& InstanceID
 
 void FProfilerManager::ProfilerClient_OnLoadedMetaData( const FGuid& InstanceID )
 {
-	TrackDefaultStats();
+	//TrackDefaultStats();
 }
 
 void FProfilerManager::ProfilerClient_OnLoadStarted( const FGuid& InstanceID )
@@ -691,7 +691,10 @@ void FProfilerManager::TrackDefaultStats()
 	for( auto It = GetProfilerInstancesIterator(); It; ++It )
 	{
 		FProfilerSessionRef ProfilerSession = It.Value();
-		TrackStat( ProfilerSession->GetMetaData()->GetGameThreadStatID() );
+		if( ProfilerSession->GetMetaData()->IsReady() )
+		{
+			TrackStat( ProfilerSession->GetMetaData()->GetGameThreadStatID() );
+		}
 		break;
 	}
 }
diff --git a/Engine/Source/Developer/Profiler/Private/ProfilerSession.h b/Engine/Source/Developer/Profiler/Private/ProfilerSession.h
--- a/Engine/Source/Developer/Profiler/Private/ProfilerSession.h
+++ b/Engine/Source/Developer/Profiler/Private/ProfilerSession.h
@@ -338,6 +338,7 @@ protected:
 	/** Constructor. */
 	FProfilerStatMetaData()
 		: SecondsPerCycle( FPlatformTime::GetSecondsPerCycle() )
+		, GameThreadID( 0 )
 	{}
 
 public:
@@ -608,6 +609,11 @@ public:
 		return RenderThreadIDs;
 	}
 
+	bool IsReady() const
+	{
+		return GameThreadID != 0;
+	}
+
 	const uint32 GetGameThreadStatID() const
 	{
 		return ThreadIDtoStatID.FindChecked( GameThreadID );
diff --git a/Engine/Source/Developer/ProfilerClient/Private/ProfilerClientManager.cpp b/Engine/Source/Developer/ProfilerClient/Private/ProfilerClientManager.cpp
--- a/Engine/Source/Developer/ProfilerClient/Private/ProfilerClientManager.cpp
+++ b/Engine/Source/Developer/ProfilerClient/Private/ProfilerClientManager.cpp
@@ -751,6 +751,8 @@ void FProfilerClientManager::BroadcastMetadataUpdate()
 {
 	while( LoadConnection->DataFrames.Num() > 0 )
 	{
+		FScopeLock ScopeLock( &(LoadConnection->CriticalSection) );
+
 		// Fire a meta data update message
 		if( LoadConnection->DataFrames[0].MetaDataUpdated )
 		{
@@ -758,7 +760,6 @@ void FProfilerClientManager::BroadcastMetadataUpdate()
 			ProfilerLoadedMetaDataDelegate.Broadcast( LoadConnection->InstanceId );
 		}
 
-		FScopeLock ScopeLock( &(LoadConnection->CriticalSection) );
 		FProfilerDataFrame& DataFrame = LoadConnection->DataFrames[0];
 		ProfilerDataDelegate.Broadcast( LoadConnection->InstanceId, DataFrame, LoadConnection->DataLoadingProgress );
 		LoadConnection->DataFrames.RemoveAt( 0 );
@@ -975,6 +976,7 @@ void FServiceConnection::GenerateAccumulators(TArray<FStatMessage>& Stats, TArra
 			const FName StatName = Stat.NameAndInfo.GetShortName();
 			if (StatName == TEXT("STAT_SecondsPerCycle"))
 			{
+				FScopeLock ScopeLock( &CriticalSection );
 				MetaData.SecondsPerCycle = Stat.GetValue_double();
 			}
 		}
@@ -1131,13 +1133,13 @@ void FServiceConnection::GenerateProfilerDataFrame()
 	DataFrame.CountAccumulators.Reset();
 	DataFrame.CycleGraphs.Reset();
 	DataFrame.FloatAccumulators.Reset();
+	DataFrame.MetaDataUpdated = false;
 
 	// get the stat stack root and the non frame stats
 	FRawStatStackNode Stack;
 	TArray<FStatMessage> NonFrameStats;
 	CurrentThreadState.UncondenseStackStats( CurrentThreadState.CurrentGameFrame, Stack, nullptr, &NonFrameStats );
 
-	// @todo updating metadata here is not thread safe?
 	// cycle graphs
 	GenerateCycleGraphs( Stack, DataFrame.CycleGraphs );
 
