Profiler - Removed unused code

diff --git a/Engine/Source/Developer/ProfilerClient/Private/ProfilerClientManager.h b/Engine/Source/Developer/ProfilerClient/Private/ProfilerClientManager.h
--- a/Engine/Source/Developer/ProfilerClient/Private/ProfilerClientManager.h
+++ b/Engine/Source/Developer/ProfilerClient/Private/ProfilerClientManager.h
@@ -124,7 +124,6 @@ class FProfilerClientManager
 	public:
 		FServiceConnection* LoadConnection;
 		FArchive* FileReader;
-		FStatsStreamHeader Header;
 
 		/** Constructor */
 		FAsyncReadWorker(FServiceConnection* InConnection, FArchive* InReader)
@@ -364,8 +363,6 @@ private:
 	bool SyncLoad();
 	/** File reader */
 	FArchive* FileReader;
-	/** Stats file header. */
-	FStatsStreamHeader Header;
 #endif
 
 	/** Holds the last time a ping was made to instances */
