SessionServices: Made public headers compilable individually

diff --git a/Engine/Source/Developer/AutomationWindow/Private/AutomationWindowModule.cpp b/Engine/Source/Developer/AutomationWindow/Private/AutomationWindowModule.cpp
--- a/Engine/Source/Developer/AutomationWindow/Private/AutomationWindowModule.cpp
+++ b/Engine/Source/Developer/AutomationWindow/Private/AutomationWindowModule.cpp
@@ -13,7 +13,7 @@ public:
 
 	// IAutomationWindowModule interface
 
-	virtual TSharedRef<class SWidget> CreateAutomationWindow( const IAutomationControllerManagerRef& AutomationController, const ISessionManagerRef& SessionManager ) override
+	virtual TSharedRef<class SWidget> CreateAutomationWindow( const IAutomationControllerManagerRef& AutomationController, const TSharedRef<ISessionManager>& SessionManager ) override
 	{
 		return SNew(SAutomationWindow, AutomationController, SessionManager);
 	}
diff --git a/Engine/Source/Developer/AutomationWindow/Private/AutomationWindowPrivatePCH.h b/Engine/Source/Developer/AutomationWindow/Private/AutomationWindowPrivatePCH.h
--- a/Engine/Source/Developer/AutomationWindow/Private/AutomationWindowPrivatePCH.h
+++ b/Engine/Source/Developer/AutomationWindow/Private/AutomationWindowPrivatePCH.h
@@ -11,7 +11,9 @@
 #include "AutomationController.h"
 #include "DesktopPlatformModule.h"
 #include "IFilter.h"
-#include "SessionServices.h"
+#include "ISessionInfo.h"
+#include "ISessionInstanceInfo.h"
+#include "ISessionManager.h"
 #include "SlateBasics.h"
 #include "EditorStyle.h"
 
diff --git a/Engine/Source/Developer/AutomationWindow/Public/AutomationWindow.h b/Engine/Source/Developer/AutomationWindow/Public/AutomationWindow.h
--- a/Engine/Source/Developer/AutomationWindow/Public/AutomationWindow.h
+++ b/Engine/Source/Developer/AutomationWindow/Public/AutomationWindow.h
@@ -9,7 +9,6 @@
 #include "AutomationController.h"
 #include "Core.h"
 #include "ModuleInterface.h"
-#include "SessionServices.h"
 #include "SlateBasics.h"
 
 /* Interfaces
diff --git a/Engine/Source/Developer/AutomationWindow/Public/Interfaces/IAutomationWindowModule.h b/Engine/Source/Developer/AutomationWindow/Public/Interfaces/IAutomationWindowModule.h
--- a/Engine/Source/Developer/AutomationWindow/Public/Interfaces/IAutomationWindowModule.h
+++ b/Engine/Source/Developer/AutomationWindow/Public/Interfaces/IAutomationWindowModule.h
@@ -3,6 +3,9 @@
 #pragma once
 
 
+class ISessionManager;
+
+
 /** Delegate to call when the automation window module is shutdown. */
 DECLARE_DELEGATE(FOnAutomationWindowModuleShutdown);
 
@@ -22,7 +25,7 @@ public:
 	 * @param SessionManager The session manager to use.
 	 * @return	New Automation Test window widget
 	 */
-	virtual TSharedRef<class SWidget> CreateAutomationWindow( const IAutomationControllerManagerRef& AutomationController, const ISessionManagerRef& SessionManager ) = 0;
+	virtual TSharedRef<class SWidget> CreateAutomationWindow( const IAutomationControllerManagerRef& AutomationController, const TSharedRef<ISessionManager>& SessionManager ) = 0;
 
 	/**
 	 * Store a pointer to the AutomationWindowTab.
diff --git a/Engine/Source/Developer/Profiler/Private/ProfilerPrivatePCH.h b/Engine/Source/Developer/Profiler/Private/ProfilerPrivatePCH.h
--- a/Engine/Source/Developer/Profiler/Private/ProfilerPrivatePCH.h
+++ b/Engine/Source/Developer/Profiler/Private/ProfilerPrivatePCH.h
@@ -8,6 +8,9 @@
 /* Private dependencies
  *****************************************************************************/
 
+#include "ISessionInfo.h"
+#include "ISessionInstanceInfo.h"
+#include "ISessionManager.h"
 #include "Messaging.h"
 #include "ProfilerClient.h"
 #include "StatsData.h"
diff --git a/Engine/Source/Developer/Profiler/Public/Interfaces/IProfilerModule.h b/Engine/Source/Developer/Profiler/Public/Interfaces/IProfilerModule.h
--- a/Engine/Source/Developer/Profiler/Public/Interfaces/IProfilerModule.h
+++ b/Engine/Source/Developer/Profiler/Public/Interfaces/IProfilerModule.h
@@ -1,11 +1,13 @@
 // Copyright 1998-2015 Epic Games, Inc. All Rights Reserved.
 
-/*=============================================================================
-	IProfilerModule.h: Declares the IProfilerModule interface.
-=============================================================================*/
-
 #pragma once
 
+
+// forward declarations
+class IProfilerManager;
+class ISessionManager;
+
+
 /**
  * Interface for the profiler module.
  */
@@ -17,15 +19,14 @@ public:
 	/**
 	 * Creates the main window for the profiler.
 	 *
-	 * @param InSessionManager			- the session manager to use
-	 * @param ConstructUnderMajorTab	- the tab which will contain the profiler tabs
+	 * @param InSessionManager The session manager to use.
+	 * @param ConstructUnderMajorTab The tab which will contain the profiler tabs.
 	 *
 	 */
-	virtual TSharedRef<SWidget> CreateProfilerWindow( const ISessionManagerRef& InSessionManager, const TSharedRef<SDockTab>& ConstructUnderMajorTab ) = 0;
-
+	virtual TSharedRef<SWidget> CreateProfilerWindow( const TSharedRef<ISessionManager>& InSessionManager, const TSharedRef<SDockTab>& ConstructUnderMajorTab ) = 0;
 
 	/**
 	 * @return a weak pointer to the global instance of the profiler manager.
 	 */
-	virtual TWeakPtr<class IProfilerManager> GetProfilerManager() = 0;
+	virtual TWeakPtr<IProfilerManager> GetProfilerManager() = 0;
 };
\ No newline at end of file
diff --git a/Engine/Source/Developer/Profiler/Public/Profiler.h b/Engine/Source/Developer/Profiler/Public/Profiler.h
--- a/Engine/Source/Developer/Profiler/Public/Profiler.h
+++ b/Engine/Source/Developer/Profiler/Public/Profiler.h
@@ -1,9 +1,5 @@
 // Copyright 1998-2015 Epic Games, Inc. All Rights Reserved.
 
-/*=============================================================================
-	Profiler.h: Profiler module public header file.
-=============================================================================*/
-
 #pragma once
 
 
@@ -14,7 +10,7 @@
 #include "ChunkedArray.h"
 #include "Ticker.h"
 #include "ModuleInterface.h"
-#include "SessionServices.h"
+#include "ISessionServicesModule.h"
 #include "SlateBasics.h"
 #include "EditorStyle.h"
 #include "IFilter.h"
diff --git a/Engine/Source/Developer/ProfilerClient/Private/ProfilerClientManager.cpp b/Engine/Source/Developer/ProfilerClient/Private/ProfilerClientManager.cpp
--- a/Engine/Source/Developer/ProfilerClient/Private/ProfilerClientManager.cpp
+++ b/Engine/Source/Developer/ProfilerClient/Private/ProfilerClientManager.cpp
@@ -5,6 +5,7 @@
 =============================================================================*/
 
 #include "ProfilerClientPrivatePCH.h"
+#include "ISessionInstanceInfo.h"
 #include "SecureHash.h"
 #include "StatsData.h"
 #include "StatsFile.h"
@@ -134,7 +135,7 @@ void FProfilerClientManager::Track( const FGuid& Instance )
 #endif
 }
 
-void FProfilerClientManager::Track( const TArray<ISessionInstanceInfoPtr>& Instances )
+void FProfilerClientManager::Track( const TArray<TSharedPtr<ISessionInstanceInfo>>& Instances )
 {
 #if STATS
 	if (MessageEndpoint.IsValid() && ActiveSessionId.IsValid())
diff --git a/Engine/Source/Developer/ProfilerClient/Private/ProfilerClientManager.h b/Engine/Source/Developer/ProfilerClient/Private/ProfilerClientManager.h
--- a/Engine/Source/Developer/ProfilerClient/Private/ProfilerClientManager.h
+++ b/Engine/Source/Developer/ProfilerClient/Private/ProfilerClientManager.h
@@ -1,13 +1,11 @@
 // Copyright 1998-2015 Epic Games, Inc. All Rights Reserved.
 
-/*=============================================================================
-	ProfilerClientManager.h: Declares the FProfilerClientManager class.
-=============================================================================*/
-
 #pragma once
 
+
 #define PROFILER_THREADED_LOAD 1
 
+
 /** Helper struct containing all of the data and operations associated with a service connection */
 struct FServiceConnection
 {
@@ -144,36 +142,26 @@ public:
 	/**
 	 * Default constructor
 	 *
-	 * @param InMessageBus - The message bus to use.
+	 * @param InMessageBus The message bus to use.
 	 */
 	FProfilerClientManager( const IMessageBusRef& InMessageBus );
 
-	/**
-	 * Default destructor
-	 */
+	/** Default destructor */
 	~FProfilerClientManager();
 
 public:
 
-	// Begin IProfileClient Interface
-	virtual void Subscribe( const FGuid& Session ) override;
+	// IProfilerClient Interface
 
+	virtual void Subscribe( const FGuid& Session ) override;
 	virtual void Track( const FGuid& Instance ) override;
-
-	virtual void Track( const TArray<ISessionInstanceInfoPtr>& Instances ) override;
-
+	virtual void Track( const TArray<TSharedPtr<ISessionInstanceInfo>>& Instances ) override;
 	virtual void Untrack( const FGuid& Instance ) override;
-
 	virtual void Unsubscribe() override;
-
 	virtual void SetCaptureState( const bool bRequestedCaptureState, const FGuid& InstanceId = FGuid() ) override;
-
 	virtual void SetPreviewState( const bool bRequestedPreviewState, const FGuid& InstanceId = FGuid() ) override;
-
 	virtual void LoadCapture( const FString& DataFilepath, const FGuid& ProfileId ) override;
-
 	virtual void RequestMetaData() override;
-
 	virtual void RequestLastCapturedFile( const FGuid& InstanceId = FGuid() ) override;
 
 	virtual const FStatMetaData& GetStatMetaData( const FGuid& InstanceId ) const override
@@ -221,7 +209,6 @@ public:
 	{
 		return ProfilerLoadedMetaDataDelegate;
 	}
-	// End IAutomationControllerModule Interface
 
 	static const int32 MaxFramesPerTick = 60;
 
diff --git a/Engine/Source/Developer/ProfilerClient/Public/Interfaces/IProfilerClient.h b/Engine/Source/Developer/ProfilerClient/Public/Interfaces/IProfilerClient.h
--- a/Engine/Source/Developer/ProfilerClient/Public/Interfaces/IProfilerClient.h
+++ b/Engine/Source/Developer/ProfilerClient/Public/Interfaces/IProfilerClient.h
@@ -1,20 +1,11 @@
 // Copyright 1998-2015 Epic Games, Inc. All Rights Reserved.
 
-/*=============================================================================
-	IProfilerClient.h: Declares the IProfilerClient interface.
-=============================================================================*/
-
 #pragma once
 
-/**
- * Type definition for shared pointers to instances of IProfilerClient.
- */
-typedef TSharedPtr<class IProfilerClient> IProfilerClientPtr;
 
-/**
- * Type definition for shared references to instances of IProfilerClient.
- */
-typedef TSharedRef<class IProfilerClient> IProfilerClientRef;
+// forward declarations
+class ISessionInstanceInfo;
+
 
 /** Delegate for passing profiler data to UI */
 DECLARE_MULTICAST_DELEGATE_ThreeParams(FProfilerClientDataDelegate, const FGuid&, const FProfilerDataFrame&, const float );
@@ -40,16 +31,18 @@ DECLARE_MULTICAST_DELEGATE_OneParam(FProfilerLoadedMetaDataDelegate, const FGuid
 /** Delegate for updating the current progress of file transfer. */
 DECLARE_MULTICAST_DELEGATE_ThreeParams(FProfilerFileTransferDelegate, const FString& /*Filename*/, int64 /*FileProgress*/, int64 /*FileSize*/ );
 
+
 /**
  * Interface for Profiler Client.
  */
 class IProfilerClient
 {
 public:
+
 	/**
 	 * Subscribe to the given profiler session id
 	 *
-	 * @param Session - session to subscribe to for profiler data
+	 * @param Session The session to subscribe to for profiler data
 	 */
 	virtual void Subscribe( const FGuid& Session ) = 0;
 
@@ -58,23 +51,21 @@ public:
 	 */
 	virtual void Track( const FGuid& Instance ) = 0;
 
-	virtual void Track( const TArray<ISessionInstanceInfoPtr>& Instances ) = 0;
+	virtual void Track( const TArray<TSharedPtr<ISessionInstanceInfo>>& Instances ) = 0;
 
 	/**
 	 * 
 	 */
 	virtual void Untrack( const FGuid& Instance ) = 0;
 
-	/**
-	 * Unsubscribe from all sessions
-	 */
+	/** Unsubscribe from all sessions. */
 	virtual void Unsubscribe() = 0;
 
 	/**
 	 * Changes the current state of the capturing data service-side.
 	 *
-	 * @param bRequestedCaptureState	- data capture state that should be set
-	 * @param InstanceId				- if valid, this function will be executed only for the specified instance, otherwise will be executed on all instances
+	 * @param bRequestedCaptureState Data capture state that should be set.
+	 * @param InstanceId If valid, this function will be executed only for the specified instance, otherwise will be executed on all instances.
 	 *
 	 */
 	virtual void SetCaptureState( const bool bRequestedCaptureState, const FGuid& InstanceId = FGuid() ) = 0; 
@@ -82,8 +73,8 @@ public:
 	/**
 	 * Changes the current state of the previewing capture data.
 	 *
-	 * @param bRequestedCaptureState	- data preview state that should be set
-	 * @param InstanceId				- if valid, this function will be executed only for the specified instance, otherwise will be executed on all instances
+	 * @param bRequestedCaptureState Data preview state that should be set.
+	 * @param InstanceId If valid, this function will be executed only for the specified instance, otherwise will be executed on all instances.
 	 *
 	 */
 	virtual void SetPreviewState( const bool bRequestedPreviewState, const FGuid& InstanceId = FGuid() ) = 0;
@@ -93,87 +84,89 @@ public:
 	 */
 	virtual void LoadCapture( const FString& DataFilepath, const FGuid& ProfileId ) = 0;
 
-	/**
-	 * Requests the meta data from the service
-	 */
+	/** Requests the meta data from the service */
 	virtual void RequestMetaData() = 0;
 
 	/**
 	 * Requests the last captured file from the service.
 	 *
-	 * @param InstanceId				- if valid, this function will be executed only for the specified instance, otherwise will be executed on all instances
-	 *
+	 * @param InstanceId If valid, this function will be executed only for the specified instance, otherwise will be executed on all instances
 	 */
 	virtual void RequestLastCapturedFile( const FGuid& InstanceId = FGuid() ) = 0; 
 
 	/**
 	 * Gets the description for the given stat id
 	 *
-	 * @param StatId - id of the statistic description to retrieve
-	 * @return the FStatMetaData struct with the description
+	 * @param StatId The identifier of the statistic description to retrieve.
+	 * @return the FStatMetaData struct with the description.
 	 */
 	virtual const FStatMetaData& GetStatMetaData( const FGuid& InstanceId ) const = 0;
 
 	/**
-	 * Retrieves the profiler data delegate
+	 * Retrieves the profiler data delegate.
 	 *
-	 * @return profiler data delegate
+	 * @return profiler data delegate.
 	 */
 	virtual FProfilerClientDataDelegate& OnProfilerData() = 0;
 
 	/**
-	 * Retrieves the profiler file transfer delegate
+	 * Retrieves the profiler file transfer delegate.
 	 *
-	 * @return profiler file transfer delegate
+	 * @return profiler file transfer delegate.
 	 */
 	virtual FProfilerFileTransferDelegate& OnProfilerFileTransfer() = 0;
 
 	/**
-	 * Retrieves the profiler client connected delegate
+	 * Retrieves the profiler client connected delegate.
 	 *
-	 * @return profiler client connected delegate
+	 * @return profiler client connected delegate.
 	 */
 	virtual FProfilerClientConnectedDelegate& OnProfilerClientConnected() = 0;
 
 	/**
-	 * Retrieves the profiler client disconnected delegate
+	 * Retrieves the profiler client disconnected delegate.
 	 *
-	 * @return profiler client disconnected delegate
+	 * @return profiler client disconnected delegate.
 	 */
 	virtual FProfilerClientDisconnectedDelegate& OnProfilerClientDisconnected() = 0;
 
 	/**
-	 * Retrieves the profiler meta data update delegate
+	 * Retrieves the profiler meta data update delegate.
 	 *
-	 * @return profiler meta data delegate
+	 * @return profiler meta data delegate.
 	 */
 	virtual FProfilerMetaDataUpdateDelegate& OnMetaDataUpdated() = 0;
 
 	/**
-	 * Retrieves the load started delegate
+	 * Retrieves the load started delegate.
 	 *
-	 * @return profiler load start delegate
+	 * @return profiler load start delegate.
 	 */
 	virtual FProfilerLoadStartedDelegate& OnLoadStarted() = 0;
 
 	/**
-	 * Retrieves the load completed delegate
+	 * Retrieves the load completed delegate.
 	 *
-	 * @return profiler load completed delegate
+	 * @return profiler load completed delegate.
 	 */
 	virtual FProfilerLoadCompletedDelegate& OnLoadCompleted() = 0;
 
 	/**
-	 * Retrieves the meta data loaded delegate
+	 * Retrieves the meta data loaded delegate.
 	 *
-	 * @return profiler meta dataloaded delegate
+	 * @return profiler meta dataloaded delegate.
 	 */
 	virtual FProfilerLoadedMetaDataDelegate& OnLoadedMetaData() = 0;
 
 protected:
 
-	/**
-	 * Virtual destructor
-	 */
+	/** Virtual destructor. */
 	virtual ~IProfilerClient() { }
 };
+
+
+/** Type definition for shared pointers to instances of IProfilerClient. */
+typedef TSharedPtr<class IProfilerClient> IProfilerClientPtr;
+
+/** Type definition for shared references to instances of IProfilerClient. */
+typedef TSharedRef<class IProfilerClient> IProfilerClientRef;
diff --git a/Engine/Source/Developer/ProfilerClient/Public/ProfilerClient.h b/Engine/Source/Developer/ProfilerClient/Public/ProfilerClient.h
--- a/Engine/Source/Developer/ProfilerClient/Public/ProfilerClient.h
+++ b/Engine/Source/Developer/ProfilerClient/Public/ProfilerClient.h
@@ -1,9 +1,5 @@
 // Copyright 1998-2015 Epic Games, Inc. All Rights Reserved.
 
-/*=============================================================================
-	ProfilerClient.h: Profiler Client module public header file.
-=============================================================================*/
-
 #pragma once
 
 
@@ -12,7 +8,6 @@
 
 #include "Core.h"
 #include "ModuleManager.h"
-#include "SessionServices.h"
 #include "ProfilerService.h"
 
 
diff --git a/Engine/Source/Developer/SessionFrontend/Private/SessionFrontendPrivatePCH.h b/Engine/Source/Developer/SessionFrontend/Private/SessionFrontendPrivatePCH.h
--- a/Engine/Source/Developer/SessionFrontend/Private/SessionFrontendPrivatePCH.h
+++ b/Engine/Source/Developer/SessionFrontend/Private/SessionFrontendPrivatePCH.h
@@ -16,6 +16,12 @@
 #include "TargetPlatform.h"
 #include "PlatformInfo.h"
 
+// SessionServices
+#include "ISessionInfo.h"
+#include "ISessionInstanceInfo.h"
+#include "ISessionManager.h"
+#include "SessionLogMessage.h"
+
 // @todo gmp: remove these dependencies by making the session front-end extensible
 #include "AutomationWindow.h"
 #include "ScreenShotComparison.h"
diff --git a/Engine/Source/Developer/SessionFrontend/SessionFrontend.Build.cs b/Engine/Source/Developer/SessionFrontend/SessionFrontend.Build.cs
--- a/Engine/Source/Developer/SessionFrontend/SessionFrontend.Build.cs
+++ b/Engine/Source/Developer/SessionFrontend/SessionFrontend.Build.cs
@@ -14,17 +14,12 @@ public class SessionFrontend : ModuleRules
 			}
 		);
 
-		PublicIncludePathModuleNames.AddRange(
-			new string[] {
-				"SessionServices",
-			}
-		);
-
-		PrivateDependencyModuleNames.AddRange(
+        PrivateDependencyModuleNames.AddRange(
 			new string[] {
 				"DesktopPlatform",
                 "InputCore",
 				"Json",
+                "SessionServices",
 				"SlateCore",
 
 				// @todo gmp: remove these dependencies by making the session front-end extensible
diff --git a/Engine/Source/Developer/SessionServices/Private/SessionInfo.cpp b/Engine/Source/Developer/SessionServices/Private/SessionInfo.cpp
new file mode 100644
--- /dev/null
+++ b/Engine/Source/Developer/SessionServices/Private/SessionInfo.cpp
@@ -0,0 +1,150 @@
+// Copyright 1998-2015 Epic Games, Inc. All Rights Reserved.
+
+#include "SessionServicesPrivatePCH.h"
+
+
+/* FSessionInfo structors
+ *****************************************************************************/
+
+FSessionInfo::FSessionInfo( const FGuid& InSessionId, const IMessageBusRef& InMessageBus )
+	: MessageBusPtr(InMessageBus)
+	, SessionId(InSessionId)
+{ }
+
+
+/* FSessionInfo interface
+ *****************************************************************************/
+
+void FSessionInfo::UpdateFromMessage( const FEngineServicePong& Message, const IMessageContextRef& Context )
+{
+	if (Message.SessionId != SessionId)
+	{
+		return;
+	}
+
+	// update instance
+	// @todo gmp: reconsider merging EngineService and SessionService
+	/*TSharedPtr<FSessionInstanceInfo> Instance = Instances.FindRef(Context->GetSender());
+
+	if (Instance.IsValid())
+	{
+		Instance->UpdateFromMessage(Message, Context);
+	}*/
+	for (TMap<FMessageAddress, TSharedPtr<FSessionInstanceInfo> >::TIterator It(Instances); It; ++It)
+	{
+		if (It.Value()->GetInstanceId() == Message.InstanceId)
+		{
+			It.Value()->UpdateFromMessage(Message, Context);
+			break;
+		}
+	}
+}
+
+
+void FSessionInfo::UpdateFromMessage( const FSessionServicePong& Message, const IMessageContextRef& Context )
+{
+	if (Message.SessionId != SessionId)
+	{
+		return;
+	}
+
+	// update session info
+	Standalone = Message.Standalone;
+	SessionOwner = Message.SessionOwner;
+
+	if (SessionName.IsEmpty())
+	{
+		SessionName = Message.SessionName;
+	}
+
+	// update instance
+	TSharedPtr<FSessionInstanceInfo>& Instance = Instances.FindOrAdd(Context->GetSender());
+
+	if (Instance.IsValid())
+	{
+		Instance->UpdateFromMessage(Message, Context);
+	}
+	else
+	{
+		IMessageBusPtr MessageBus = MessageBusPtr.Pin();
+
+		if (MessageBus.IsValid())
+		{
+			Instance = MakeShareable(new FSessionInstanceInfo(Message.InstanceId, AsShared(), MessageBus.ToSharedRef()));
+			Instance->OnLogReceived().AddSP(this, &FSessionInfo::HandleLogReceived);
+			Instance->UpdateFromMessage(Message, Context);
+
+			InstanceDiscoveredEvent.Broadcast(AsShared(), Instance.ToSharedRef());
+		}
+	}
+
+	LastUpdateTime = FDateTime::UtcNow();
+}
+
+
+/* ISessionInfo interface
+ *****************************************************************************/
+
+void FSessionInfo::GetInstances( TArray<ISessionInstanceInfoPtr>& OutInstances ) const
+{
+	OutInstances.Empty();
+
+	for (TMap<FMessageAddress, TSharedPtr<FSessionInstanceInfo> >::TConstIterator It(Instances); It; ++It)
+	{
+		OutInstances.Add(It.Value());
+	}
+}
+
+
+const FDateTime& FSessionInfo::GetLastUpdateTime() const
+{
+	return LastUpdateTime;
+}
+
+
+const int32 FSessionInfo::GetNumInstances() const
+{
+	return Instances.Num();
+}
+
+
+const FGuid& FSessionInfo::GetSessionId() const
+{
+	return SessionId;
+}
+
+
+const FString& FSessionInfo::GetSessionName() const
+{
+	return SessionName;
+}
+
+
+const FString& FSessionInfo::GetSessionOwner() const
+{
+	return SessionOwner;
+}
+
+
+const bool FSessionInfo::IsStandalone() const
+{
+	return Standalone;
+}
+
+
+void FSessionInfo::Terminate()
+{
+	for (TMap<FMessageAddress, TSharedPtr<FSessionInstanceInfo> >::TIterator It(Instances); It; ++It)
+	{
+		It.Value()->Terminate();
+	}
+}
+
+
+/* FSessionInfo callbacks
+ *****************************************************************************/
+
+void FSessionInfo::HandleLogReceived( const ISessionInstanceInfoRef& Instance, const FSessionLogMessageRef& LogMessage )
+{
+	LogReceivedEvent.Broadcast(AsShared(), Instance, LogMessage);
+}
diff --git a/Engine/Source/Developer/SessionServices/Private/SessionInfo.h b/Engine/Source/Developer/SessionServices/Private/SessionInfo.h
--- a/Engine/Source/Developer/SessionServices/Private/SessionInfo.h
+++ b/Engine/Source/Developer/SessionServices/Private/SessionInfo.h
@@ -18,10 +18,7 @@ public:
 	 * @param InSessionId The session's identifier.
 	 * @param InMessageBus The message bus to use.
 	 */
-	FSessionInfo( const FGuid& InSessionId, const IMessageBusRef& InMessageBus )
-		: MessageBusPtr(InMessageBus)
-		, SessionId(InSessionId)
-	{ }
+	FSessionInfo( const FGuid& InSessionId, const IMessageBusRef& InMessageBus );
 
 public:
 
@@ -31,30 +28,7 @@ public:
 	 * @param Message The message containing engine information.
 	 * @param Context The message context.
 	 */
-	void UpdateFromMessage( const FEngineServicePong& Message, const IMessageContextRef& Context )
-	{
-		if (Message.SessionId != SessionId)
-		{
-			return;
-		}
-
-		// update instance
-		// @todo gmp: reconsider merging EngineService and SessionService
-		/*TSharedPtr<FSessionInstanceInfo> Instance = Instances.FindRef(Context->GetSender());
-
-		if (Instance.IsValid())
-		{
-			Instance->UpdateFromMessage(Message, Context);
-		}*/
-		for (TMap<FMessageAddress, TSharedPtr<FSessionInstanceInfo> >::TIterator It(Instances); It; ++It)
-		{
-			if (It.Value()->GetInstanceId() == Message.InstanceId)
-			{
-				It.Value()->UpdateFromMessage(Message, Context);
-				break;
-			}
-		}
-	}
+	void UpdateFromMessage( const FEngineServicePong& Message, const IMessageContextRef& Context );
 
 	/**
 	 * Updates this session info with the data in the specified message.
@@ -62,115 +36,38 @@ public:
 	 * @param Message The message containing session information.
 	 * @param Context The message context.
 	 */
-	void UpdateFromMessage( const FSessionServicePong& Message, const IMessageContextRef& Context )
-	{
-		if (Message.SessionId != SessionId)
-		{
-			return;
-		}
-
-		// update session info
-		Standalone = Message.Standalone;
-		SessionOwner = Message.SessionOwner;
-
-		if (SessionName.IsEmpty())
-		{
-			SessionName = Message.SessionName;
-		}
-
-		// update instance
-		TSharedPtr<FSessionInstanceInfo>& Instance = Instances.FindOrAdd(Context->GetSender());
-
-		if (Instance.IsValid())
-		{
-			Instance->UpdateFromMessage(Message, Context);
-		}
-		else
-		{
-			IMessageBusPtr MessageBus = MessageBusPtr.Pin();
-
-			if (MessageBus.IsValid())
-			{
-				Instance = MakeShareable(new FSessionInstanceInfo(Message.InstanceId, AsShared(), MessageBus.ToSharedRef()));
-				Instance->OnLogReceived().AddSP(this, &FSessionInfo::HandleLogReceived);
-				Instance->UpdateFromMessage(Message, Context);
-
-				InstanceDiscoveredDelegate.Broadcast(AsShared(), Instance.ToSharedRef());
-			}
-		}
-
-		LastUpdateTime = FDateTime::UtcNow();
-	}
+	void UpdateFromMessage( const FSessionServicePong& Message, const IMessageContextRef& Context );
 
 public:	
 
-	// SessionInfo interface
+	// ISessionInfo interface
 
-	virtual void GetInstances( TArray<ISessionInstanceInfoPtr>& OutInstances ) const override
-	{
-		OutInstances.Empty();
-
-		for (TMap<FMessageAddress, TSharedPtr<FSessionInstanceInfo> >::TConstIterator It(Instances); It; ++It)
-		{
-			OutInstances.Add(It.Value());
-		}
-	}
+	virtual void GetInstances( TArray<ISessionInstanceInfoPtr>& OutInstances ) const override;
+	virtual const FDateTime& GetLastUpdateTime() const override;
+	virtual const int32 GetNumInstances() const override;
+	virtual const FGuid& GetSessionId() const override;
+	virtual const FString& GetSessionName() const override;
+	virtual const FString& GetSessionOwner() const override;
+	virtual const bool IsStandalone() const override;
 
-	virtual FDateTime GetLastUpdateTime() override
+	DECLARE_DERIVED_EVENT(FSessionInfo, ISessionInfo::FInstanceDiscoveredEvent, FInstanceDiscoveredEvent)
+	virtual FInstanceDiscoveredEvent& OnInstanceDiscovered() override
 	{
-		return LastUpdateTime;
+		return InstanceDiscoveredEvent;
 	}
 
-	virtual const int32 GetNumInstances() const override
+	DECLARE_DERIVED_EVENT(FSessionInfo, ISessionInfo::FLogReceivedEvent, FLogReceivedEvent)
+	virtual FLogReceivedEvent& OnLogReceived() override
 	{
-		return Instances.Num();
+		return LogReceivedEvent;
 	}
 
-	virtual const FGuid& GetSessionId() const override
-	{
-		return SessionId;
-	}
-
-	virtual const FString& GetSessionName() const override
-	{
-		return SessionName;
-	}
-
-	virtual const FString& GetSessionOwner() const override
-	{
-		return SessionOwner;
-	}
-
-	virtual const bool IsStandalone() const override
-	{
-		return Standalone;
-	}
-
-	virtual FOnSessionInstanceDiscovered& OnInstanceDiscovered() override
-	{
-		return InstanceDiscoveredDelegate;
-	}
-
-	virtual FOnSessionLogReceived& OnLogReceived() override
-	{
-		return LogReceivedDelegate;
-	}
-
-	virtual void Terminate() override
-	{
-		for (TMap<FMessageAddress, TSharedPtr<FSessionInstanceInfo> >::TIterator It(Instances); It; ++It)
-		{
-			It.Value()->Terminate();
-		}
-	}
+	virtual void Terminate() override;
 
 private:
 
 	/** Handles received log messages. */
-	void HandleLogReceived( const ISessionInstanceInfoRef& Instance, const FSessionLogMessageRef& LogMessage )
-	{
-		LogReceivedDelegate.Broadcast(AsShared(), Instance, LogMessage);
-	}
+	void HandleLogReceived( const ISessionInstanceInfoRef& Instance, const FSessionLogMessageRef& LogMessage );
 
 private:
 
@@ -198,8 +95,8 @@ private:
 private:
 
 	/** Holds a delegate to be invoked when a new instance has been discovered. */
-	FOnSessionInstanceDiscovered InstanceDiscoveredDelegate;
+	FInstanceDiscoveredEvent InstanceDiscoveredEvent;
 
 	/** Holds a delegate to be invoked when an instance received a log message. */
-	FOnSessionLogReceived LogReceivedDelegate;
+	FLogReceivedEvent LogReceivedEvent;
 };
diff --git a/Engine/Source/Developer/SessionServices/Private/SessionInstanceInfo.cpp b/Engine/Source/Developer/SessionServices/Private/SessionInstanceInfo.cpp
new file mode 100644
--- /dev/null
+++ b/Engine/Source/Developer/SessionServices/Private/SessionInstanceInfo.cpp
@@ -0,0 +1,100 @@
+// Copyright 1998-2015 Epic Games, Inc. All Rights Reserved.
+
+#include "SessionServicesPrivatePCH.h"
+
+
+/* FSessionInstanceInfo structors
+ *****************************************************************************/
+
+FSessionInstanceInfo::FSessionInstanceInfo( const FGuid& InInstanceId, const TSharedRef<ISessionInfo>& InOwner, const IMessageBusRef& InMessageBus )
+	: EngineVersion(0)
+	, InstanceId(InInstanceId)
+	, Owner(InOwner)
+{
+	MessageEndpoint = FMessageEndpoint::Builder("FSessionInstanceInfo", InMessageBus)
+		.Handling<FSessionServiceLog>(this, &FSessionInstanceInfo::HandleSessionLogMessage);
+}
+
+
+/* FSessionInstanceInfo interface
+ *****************************************************************************/
+
+void FSessionInstanceInfo::UpdateFromMessage( const FEngineServicePong& Message, const IMessageContextRef& Context )
+{
+	if (Message.InstanceId != InstanceId)
+	{
+		return;
+	}
+
+	CurrentLevel = Message.CurrentLevel;
+	EngineAddress = Context->GetSender();
+	EngineVersion = Message.EngineVersion;
+	HasBegunPlay = Message.HasBegunPlay;
+	WorldTimeSeconds = Message.WorldTimeSeconds;
+	InstanceType = Message.InstanceType;
+}
+
+
+void FSessionInstanceInfo::UpdateFromMessage( const FSessionServicePong& Message, const IMessageContextRef& Context )
+{
+	if (Message.InstanceId != InstanceId)
+	{
+		return;
+	}
+
+	if (MessageEndpoint.IsValid() && (ApplicationAddress != Context->GetSender()))
+	{
+		MessageEndpoint->Send(new FSessionServiceLogSubscribe(), Context->GetSender());
+	}
+
+	ApplicationAddress = Context->GetSender();
+	BuildDate = Message.BuildDate;
+	DeviceName = Message.DeviceName;
+	InstanceName = Message.InstanceName;
+	IsConsoleBuild = Message.IsConsoleBuild;
+	PlatformName = Message.PlatformName;
+
+	LastUpdateTime = FDateTime::UtcNow();
+}
+
+
+/* FSessionInstanceInfo interface
+ *****************************************************************************/
+
+void FSessionInstanceInfo::ExecuteCommand( const FString& CommandString )
+{
+	if (MessageEndpoint.IsValid() && EngineAddress.IsValid())
+	{
+		MessageEndpoint->Send(new FEngineServiceExecuteCommand(CommandString, FPlatformProcess::UserName(true)), EngineAddress);
+	}
+}
+
+
+void FSessionInstanceInfo::Terminate()
+{
+	if (MessageEndpoint.IsValid() && EngineAddress.IsValid())
+	{
+		MessageEndpoint->Send(new FEngineServiceTerminate(FPlatformProcess::UserName(true)), EngineAddress);
+	}
+}
+
+
+/* FSessionInstanceInfo callbacks
+ *****************************************************************************/
+
+void FSessionInstanceInfo::HandleSessionLogMessage( const FSessionServiceLog& Message, const IMessageContextRef& Context )
+{
+	FSessionLogMessageRef LogMessage = MakeShareable(
+		new FSessionLogMessage(
+			InstanceId,
+			InstanceName,
+			Message.TimeSeconds,
+			Message.Data,
+			(ELogVerbosity::Type)Message.Verbosity,
+			Message.Category
+		)
+	);
+
+	LogMessages.Add(LogMessage);
+	LogReceivedEvent.Broadcast(AsShared(), LogMessage);
+}
diff --git a/Engine/Source/Developer/SessionServices/Private/SessionInstanceInfo.h b/Engine/Source/Developer/SessionServices/Private/SessionInstanceInfo.h
--- a/Engine/Source/Developer/SessionServices/Private/SessionInstanceInfo.h
+++ b/Engine/Source/Developer/SessionServices/Private/SessionInstanceInfo.h
@@ -2,6 +2,12 @@
 
 #pragma once
 
+#include "ISessionInstanceInfo.h"
+
+
+class ISessionInfo;
+struct FSessionLogMessage;
+
 
 /**
  * Implements a class to maintain all info related to a game instance in a session
@@ -22,14 +28,7 @@ public:
 	 * @param InOwner The session that owns this instance.
 	 * @param InMessageBus The message bus to use.
 	 */
-	FSessionInstanceInfo( const FGuid& InInstanceId, const ISessionInfoRef& InOwner, const IMessageBusRef& InMessageBus )
-		: EngineVersion(0)
-		, InstanceId(InInstanceId)
-		, Owner(InOwner)
-	{
-		MessageEndpoint = FMessageEndpoint::Builder("FSessionInstanceInfo", InMessageBus)
-			.Handling<FSessionServiceLog>(this, &FSessionInstanceInfo::HandleSessionLogMessage);
-	}
+	FSessionInstanceInfo( const FGuid& InInstanceId, const TSharedRef<ISessionInfo>& InOwner, const IMessageBusRef& InMessageBus );
 
 public:
 
@@ -39,20 +38,7 @@ public:
 	 * @param Message The message containing engine information.
 	 * @param Context The message context.
 	 */
-	void UpdateFromMessage( const FEngineServicePong& Message, const IMessageContextRef& Context )
-	{
-		if (Message.InstanceId != InstanceId)
-		{
-			return;
-		}
-
-		CurrentLevel = Message.CurrentLevel;
-		EngineAddress = Context->GetSender();
-		EngineVersion = Message.EngineVersion;
-		HasBegunPlay = Message.HasBegunPlay;
-		WorldTimeSeconds = Message.WorldTimeSeconds;
-		InstanceType = Message.InstanceType;
-	}
+	void UpdateFromMessage( const FEngineServicePong& Message, const IMessageContextRef& Context );
 
 	/**
 	 * Updates this instance info with the data in the specified message.
@@ -60,51 +46,25 @@ public:
 	 * @param Message The message containing instance information.
 	 * @param Context The message context.
 	 */
-	void UpdateFromMessage( const FSessionServicePong& Message, const IMessageContextRef& Context )
-	{
-		if (Message.InstanceId != InstanceId)
-		{
-			return;
-		}
-
-		if (MessageEndpoint.IsValid() && (ApplicationAddress != Context->GetSender()))
-		{
-			MessageEndpoint->Send(new FSessionServiceLogSubscribe(), Context->GetSender());
-		}
-
-		ApplicationAddress = Context->GetSender();
-		BuildDate = Message.BuildDate;
-		DeviceName = Message.DeviceName;
-		InstanceName = Message.InstanceName;
-		IsConsoleBuild = Message.IsConsoleBuild;
-		PlatformName = Message.PlatformName;
-
-		LastUpdateTime = FDateTime::UtcNow();
-	}
+	void UpdateFromMessage( const FSessionServicePong& Message, const IMessageContextRef& Context );
 
 public:	
 
 	// IGameInstanceInfo interface
 
-	virtual void ExecuteCommand( const FString& CommandString )
-	{
-		if (MessageEndpoint.IsValid() && EngineAddress.IsValid())
-		{
-			MessageEndpoint->Send(new FEngineServiceExecuteCommand(CommandString, FPlatformProcess::UserName(true)), EngineAddress);
-		}
-	}
+	virtual void ExecuteCommand( const FString& CommandString ) override;
 
-	virtual const FString GetBuildDate() const override
+	virtual const FString& GetBuildDate() const override
 	{
 		return BuildDate;
 	}
 
-	virtual const FString GetCurrentLevel() const override
+	virtual const FString& GetCurrentLevel() const override
 	{
 		return CurrentLevel;
 	}
 
-	virtual const FString GetDeviceName() const override
+	virtual const FString& GetDeviceName() const override
 	{
 		return DeviceName;
 	}
@@ -114,32 +74,32 @@ public:
 		return EngineVersion;
 	}
 
-	virtual const FGuid GetInstanceId() const override
+	virtual const FGuid& GetInstanceId() const override
 	{
 		return InstanceId;
 	}
 
-	virtual const FString GetInstanceName() const override
+	virtual const FString& GetInstanceName() const override
 	{
 		return InstanceName;
 	}
 
-	virtual const FString GetInstanceType() const override
+	virtual const FString& GetInstanceType() const override
 	{
 		return InstanceType;
 	}
 
-	virtual FDateTime GetLastUpdateTime() override
+	virtual const FDateTime& GetLastUpdateTime() const override
 	{
 		return LastUpdateTime;
 	}
 
-	virtual const TArray<FSessionLogMessagePtr>& GetLog() override
+	virtual const TArray<TSharedPtr<FSessionLogMessage>>& GetLog() override
 	{
 		return LogMessages;
 	}
 
-	virtual ISessionInfoPtr GetOwnerSession() override
+	virtual TSharedPtr<ISessionInfo> GetOwnerSession() override
 	{
 		return Owner.Pin();
 	}
@@ -159,9 +119,10 @@ public:
 		return IsConsoleBuild;
 	}
 
-	virtual FOnSessionInstanceLogReceived& OnLogReceived() override
+	DECLARE_DERIVED_EVENT(FSessionInstanceInfo, ISessionInstanceInfo::FLogReceivedEvent, FLogReceivedEvent);
+	virtual FLogReceivedEvent& OnLogReceived() override
 	{
-		return LogReceivedDelegate;
+		return LogReceivedEvent;
 	}
 
 	virtual bool PlayHasBegun() const override
@@ -169,33 +130,12 @@ public:
 		return HasBegunPlay;
 	}
 
-	virtual void Terminate() override
-	{
-		if (MessageEndpoint.IsValid() && EngineAddress.IsValid())
-		{
-			MessageEndpoint->Send(new FEngineServiceTerminate(FPlatformProcess::UserName(true)), EngineAddress);
-		}
-	}
+	virtual void Terminate() override;
 
 private:
 
-	// Handles FSessionServiceLog messages.
-	void HandleSessionLogMessage( const FSessionServiceLog& Message, const IMessageContextRef& Context )
-	{
-		FSessionLogMessageRef LogMessage = MakeShareable(
-			new FSessionLogMessage(
-				InstanceId,
-				InstanceName,
-				Message.TimeSeconds,
-				Message.Data,
-				(ELogVerbosity::Type)Message.Verbosity,
-				Message.Category
-			)
-		);
-
-		LogMessages.Add(LogMessage);
-		LogReceivedDelegate.Broadcast(AsShared(), LogMessage);
-	}
+	/** Handles FSessionServiceLog messages. */
+	void HandleSessionLogMessage( const FSessionServiceLog& Message, const IMessageContextRef& Context );
 
 private:
 
@@ -236,7 +176,7 @@ private:
 	FDateTime LastUpdateTime;
 
 	/** Holds the collection of received log messages. */
-	TArray<FSessionLogMessagePtr> LogMessages;
+	TArray<TSharedPtr<FSessionLogMessage>> LogMessages;
 
 	/** Holds the message endpoint. */
 	FMessageEndpointPtr MessageEndpoint;
@@ -252,6 +192,6 @@ private:
 
 private:
 
-	/** Holds a delegate to be invoked when a log message was received from the remote session. */
-	FOnSessionInstanceLogReceived LogReceivedDelegate;
+	/** Holds an event delegate that is executed when a log message was received from the remote session. */
+	FLogReceivedEvent LogReceivedEvent;
 };
diff --git a/Engine/Source/Developer/SessionServices/Private/SessionManager.cpp b/Engine/Source/Developer/SessionServices/Private/SessionManager.cpp
--- a/Engine/Source/Developer/SessionServices/Private/SessionManager.cpp
+++ b/Engine/Source/Developer/SessionServices/Private/SessionManager.cpp
@@ -73,6 +73,12 @@ void FSessionManager::GetSelectedInstances( TArray<ISessionInstanceInfoPtr>& Out
 }
 
 
+const ISessionInfoPtr& FSessionManager::GetSelectedSession() const
+{
+	return SelectedSession;
+}
+
+
 void FSessionManager::GetSessions( TArray<ISessionInfoPtr>& OutSessions ) const
 {
 	OutSessions.Empty(Sessions.Num());
@@ -84,6 +90,24 @@ void FSessionManager::GetSessions( TArray<ISessionInfoPtr>& OutSessions ) const
 }
 
 
+bool FSessionManager::IsInstanceSelected( const TSharedRef<ISessionInstanceInfo>& Instance ) const
+{
+	return ((Instance->GetOwnerSession() == SelectedSession) && !DeselectedInstances.Contains(Instance));
+}
+
+
+FSimpleMulticastDelegate& FSessionManager::OnSessionsUpdated()
+{
+	return SessionsUpdatedDelegate;
+}
+
+
+FSimpleMulticastDelegate& FSessionManager::OnSessionInstanceUpdated()
+{
+	return SessionInstanceUpdatedDelegate;
+}
+
+
 void FSessionManager::RemoveOwner( const FString& InOwner )
 {
 	FilteredOwners.Remove(InOwner);
@@ -105,7 +129,7 @@ bool FSessionManager::SelectSession( const ISessionInfoPtr& Session )
 			{
 				SelectedSession = Session;
 
-				SelectedSessionChangedDelegate.Broadcast(Session);
+				SelectedSessionChangedEvent.Broadcast(Session);
 			}
 			else
 			{
@@ -222,7 +246,7 @@ void FSessionManager::SendPing()
 }
 
 
-/* FSessionManager event handlers
+/* FSessionManager callbacks
  *****************************************************************************/
 
 void FSessionManager::HandleEnginePongMessage( const FEngineServicePong& Message, const IMessageContextRef& Context )
@@ -256,7 +280,7 @@ void FSessionManager::HandleLogReceived( const ISessionInfoRef& Session, const I
 {
 	if (Session == SelectedSession)
 	{
-		LogReceivedDelegate.Broadcast(Session, Instance, Message);
+		LogReceivedEvent.Broadcast(Session, Instance, Message);
 	}
 }
 
diff --git a/Engine/Source/Developer/SessionServices/Private/SessionManager.h b/Engine/Source/Developer/SessionServices/Private/SessionManager.h
--- a/Engine/Source/Developer/SessionServices/Private/SessionManager.h
+++ b/Engine/Source/Developer/SessionServices/Private/SessionManager.h
@@ -2,6 +2,8 @@
 
 #pragma once
 
+#include "ISessionManager.h"
+
 
 /**
  * Implement the session manager
@@ -27,50 +29,37 @@ public:
 	// Begin ISessionManager Interface
 
 	virtual void AddOwner( const FString& InOwner ) override;
-	virtual void GetSelectedInstances( TArray<ISessionInstanceInfoPtr>& OutInstances) const override;
-
-	virtual const ISessionInfoPtr& GetSelectedSession() const override
-	{
-		return SelectedSession;
-	}
-
-	virtual void GetSessions( TArray<ISessionInfoPtr>& OutSessions ) const override;
-
-	virtual bool IsInstanceSelected( const ISessionInstanceInfoRef& Instance ) const override
-	{
-		return ((Instance->GetOwnerSession() == SelectedSession) && !DeselectedInstances.Contains(Instance));
-	}
+	virtual void GetSelectedInstances( TArray<TSharedPtr<ISessionInstanceInfo>>& OutInstances) const override;
+	virtual const ISessionInfoPtr& GetSelectedSession() const override;
+	virtual void GetSessions( TArray<TSharedPtr<ISessionInfo>>& OutSessions ) const override;
+	virtual bool IsInstanceSelected( const TSharedRef<ISessionInstanceInfo>& Instance ) const override;
 
-	virtual FOnCanSelectSession& OnCanSelectSession() override
+	DECLARE_DERIVED_EVENT(FSessionManager, ISessionManager::FCanSelectSessionEvent, FCanSelectSessionEvent)
+	virtual FCanSelectSessionEvent& OnCanSelectSession() override
 	{
 		return CanSelectSessionDelegate;
 	}
 
-	virtual FOnSessionInstanceSelectionChanged& OnInstanceSelectionChanged() override
+	DECLARE_DERIVED_EVENT(FSessionManager, ISessionManager::FInstanceSelectionChangedEvent, FInstanceSelectionChangedEvent)
+	virtual FInstanceSelectionChangedEvent& OnInstanceSelectionChanged() override
 	{
 		return InstanceSelectionChangedDelegate;
 	}
 
-	virtual FOnSessionLogReceived& OnLogReceived() override
-	{
-		return LogReceivedDelegate;
-	}
-
-	virtual FOnSelectedSessionChanged& OnSelectedSessionChanged() override
-	{
-		return SelectedSessionChangedDelegate;
-	}
-
-	virtual FSimpleMulticastDelegate& OnSessionsUpdated() override
+	DECLARE_DERIVED_EVENT(FSessionManager, ISessionManager::FLogReceivedEvent, FLogReceivedEvent)
+	virtual FLogReceivedEvent& OnLogReceived() override
 	{
-		return SessionsUpdatedDelegate;
+		return LogReceivedEvent;
 	}
 
-	virtual FSimpleMulticastDelegate& OnSessionInstanceUpdated() override
+	DECLARE_DERIVED_EVENT(FSessionManager, ISessionManager::FSelectedSessionChangedEvent, FSelectedSessionChangedEvent)
+	virtual FSelectedSessionChangedEvent& OnSelectedSessionChanged() override
 	{
-		return SessionInstanceUpdatedDelegate;
+		return SelectedSessionChangedEvent;
 	}
 
+	virtual FSimpleMulticastDelegate& OnSessionsUpdated() override;
+	virtual FSimpleMulticastDelegate& OnSessionInstanceUpdated() override;
 	virtual void RemoveOwner( const FString& InOwner ) override;
 	virtual bool SelectSession( const ISessionInfoPtr& Session ) override;
 	virtual bool SetInstanceSelected( const ISessionInstanceInfoPtr& Instance, bool Selected ) override;
@@ -140,19 +129,19 @@ private:
 private:
 
 	/** Holds a delegate to be invoked before a session is selected. */
-	FOnCanSelectSession CanSelectSessionDelegate;
+	FCanSelectSessionEvent CanSelectSessionDelegate;
 
 	/** Holds a delegate to be invoked when an instance changes its selection state. */
-	FOnSessionInstanceSelectionChanged InstanceSelectionChangedDelegate;
+	FInstanceSelectionChangedEvent InstanceSelectionChangedDelegate;
 
 	/** Owner filter list. */
 	TArray<FString> FilteredOwners;
 
 	/** Holds a delegate to be invoked when the selected session received a log message. */
-	FOnSessionLogReceived LogReceivedDelegate;
+	FLogReceivedEvent LogReceivedEvent;
 
 	/** Holds a delegate to be invoked when the selected session changed. */
-	FOnSelectedSessionChanged SelectedSessionChangedDelegate;
+	FSelectedSessionChangedEvent SelectedSessionChangedEvent;
 
 	/** Holds a delegate to be invoked when the session list was updated. */
 	FSimpleMulticastDelegate SessionsUpdatedDelegate;
diff --git a/Engine/Source/Developer/SessionServices/Private/SessionService.h b/Engine/Source/Developer/SessionServices/Private/SessionService.h
--- a/Engine/Source/Developer/SessionServices/Private/SessionService.h
+++ b/Engine/Source/Developer/SessionServices/Private/SessionService.h
@@ -2,6 +2,8 @@
 
 #pragma once
 
+#include "ISessionService.h"
+
 
 /**
  * Implements an application session service.
diff --git a/Engine/Source/Developer/SessionServices/Private/SessionServicesModule.cpp b/Engine/Source/Developer/SessionServices/Private/SessionServicesModule.cpp
--- a/Engine/Source/Developer/SessionServices/Private/SessionServicesModule.cpp
+++ b/Engine/Source/Developer/SessionServices/Private/SessionServicesModule.cpp
@@ -1,6 +1,7 @@
 // Copyright 1998-2015 Epic Games, Inc. All Rights Reserved.
 
 #include "SessionServicesPrivatePCH.h"
+#include "ISessionServicesModule.h"
 #include "ModuleManager.h"
 
 
diff --git a/Engine/Source/Developer/SessionServices/Private/SessionServicesPrivatePCH.h b/Engine/Source/Developer/SessionServices/Private/SessionServicesPrivatePCH.h
--- a/Engine/Source/Developer/SessionServices/Private/SessionServicesPrivatePCH.h
+++ b/Engine/Source/Developer/SessionServices/Private/SessionServicesPrivatePCH.h
@@ -2,8 +2,6 @@
 
 #pragma once
 
-#include "SessionServices.h"
-
 
 /* Dependencies
  *****************************************************************************/
@@ -18,6 +16,10 @@
 /* Private includes
  *****************************************************************************/
 
+#include "ISessionInfo.h"
+#include "ISessionInstanceInfo.h"
+#include "SessionLogMessage.h"
+
 #include "SessionInstanceInfo.h"
 #include "SessionInfo.h"
 #include "SessionManager.h"
diff --git a/Engine/Source/Developer/SessionServices/Public/ISessionInfo.h b/Engine/Source/Developer/SessionServices/Public/ISessionInfo.h
new file mode 100644
--- /dev/null
+++ b/Engine/Source/Developer/SessionServices/Public/ISessionInfo.h
@@ -0,0 +1,102 @@
+// Copyright 1998-2015 Epic Games, Inc. All Rights Reserved.
+
+#pragma once
+
+
+class ISessionInstanceInfo;
+struct FSessionLogMessage;
+
+
+/**
+ * Interface for game instances.
+ */
+class ISessionInfo
+{
+public:
+
+	/**
+	 * Gets a read-only collection of all instances that belong to this session.
+	 *
+	 * @param OutInstances Will hold the collection of instances.
+	 */
+	virtual void GetInstances( TArray<TSharedPtr<ISessionInstanceInfo>>& OutInstances ) const = 0;
+
+	/**
+	 * Gets the time at which the last update was received from this instance.
+	 *
+	 * @return The receive time.
+	 */
+	virtual const FDateTime& GetLastUpdateTime() const = 0;
+
+	/**
+	 * Gets the number of engine instances that are part of the session.
+	 *
+	 * @return The number of engine instances.
+	 */
+	virtual const int32 GetNumInstances() const = 0;
+
+	/**
+	 * Gets the session identifier.
+	 *
+	 * @return Session identifier.
+	 */
+	virtual const FGuid& GetSessionId() const = 0;
+
+	/**
+	 * Gets the name of the session.
+	 *
+	 * @return Session name.
+	 */
+	virtual const FString& GetSessionName() const = 0;
+
+	/**
+	 * Gets the name of the user that owns the session.
+	 *
+	 * @return User name.
+	 */
+	virtual const FString& GetSessionOwner() const = 0;
+
+	/**
+	 * Checks whether this is a standalone session.
+	 *
+	 * A session is standalone if has not been created from the Launcher.
+	 *
+	 * @return true if this is a standalone session, false otherwise.
+	 */
+	virtual const bool IsStandalone() const = 0;
+
+	/**
+	 * Terminates the session.
+	 */
+	virtual void Terminate() = 0;
+
+public:
+
+	/**
+	 * A delegate that is executed when a new instance has been discovered.
+	 *
+	 * @return The delegate.
+	 */
+	DECLARE_EVENT_TwoParams(ISessionInfo, FInstanceDiscoveredEvent, const TSharedRef<ISessionInfo>& /*OwnerSession*/, const TSharedRef<ISessionInstanceInfo>& /*DiscoveredInstance*/)
+	virtual FInstanceDiscoveredEvent& OnInstanceDiscovered() = 0;
+
+	/**
+	 * A delegate that is executed when a new log message has been received.
+	 *
+	 * @return The delegate.
+	 */
+	DECLARE_EVENT_ThreeParams(ISessionInfo, FLogReceivedEvent, const TSharedRef<ISessionInfo>& /*OwnerSession*/, const TSharedRef<ISessionInstanceInfo>& /*SessionInstance*/, const TSharedRef<FSessionLogMessage>& /*LogMessage*/);
+	virtual FLogReceivedEvent& OnLogReceived() = 0;
+
+public:
+
+	/** Virtual destructor. */
+	virtual ~ISessionInfo() { }
+};
+
+
+/** Type definition for shared pointers to instances of ISessionInfo. */
+typedef TSharedPtr<ISessionInfo> ISessionInfoPtr;
+
+/** Type definition for shared references to instances of ISessionInfo. */
+typedef TSharedRef<ISessionInfo> ISessionInfoRef;
diff --git a/Engine/Source/Developer/SessionServices/Public/ISessionInstanceInfo.h b/Engine/Source/Developer/SessionServices/Public/ISessionInstanceInfo.h
new file mode 100644
--- /dev/null
+++ b/Engine/Source/Developer/SessionServices/Public/ISessionInstanceInfo.h
@@ -0,0 +1,146 @@
+// Copyright 1998-2015 Epic Games, Inc. All Rights Reserved.
+
+#pragma once
+
+
+class ISessionInfo;
+struct FSessionLogMessage;
+
+
+/**
+ * Interface for game instances.
+ */
+class ISessionInstanceInfo
+{
+public:
+
+	/**
+	 * Executes a console command on this engine instance.
+	 *
+	 * @param CommandString The command to execute.
+	 */
+	virtual void ExecuteCommand( const FString& CommandString ) = 0;
+
+	/**
+	 * Gets instance's build date.
+	 *
+	 * @return The build date string.
+	 */
+	virtual const FString& GetBuildDate() const = 0;
+
+	/**
+	 * Gets the name of the level that the instance is currently running.
+	 *
+	 * @return The name of the level.
+	 */
+	virtual const FString& GetCurrentLevel() const = 0;
+
+	/**
+	 * Gets the name of the device that this instance is running on.
+	 *
+	 * @return The device name string.
+	 */
+	virtual const FString& GetDeviceName() const = 0;
+
+	/**
+	 * Gets the instance's engine version number.
+	 *
+	 * @return Engine version.
+	 */
+	virtual int32 GetEngineVersion() const = 0;
+
+	/**
+	 * Gets the instance identifier.
+	 *
+	 * @return Instance identifier.
+	 */
+	virtual const FGuid& GetInstanceId() const = 0;
+
+	/**
+	 * Gets the name of this instance.
+	 *
+	 * @return The instance name string.
+	 */
+	virtual const FString& GetInstanceName() const = 0;
+
+	/**
+	 * Gets the instance type (i.e. Editor or Game).
+	 *
+	 * @return The game instance type string.
+	 */
+	virtual const FString& GetInstanceType() const = 0;
+
+	/**
+	 * Gets the time at which the last update was received from this instance.
+	 *
+	 * @return The receive time.
+	 */
+	virtual const FDateTime& GetLastUpdateTime() const = 0;
+
+	/**
+	 * Gets the collection of log entries received from this instance.
+	 *
+	 * @return Log entries.
+	 */
+	virtual const TArray<TSharedPtr<FSessionLogMessage>>& GetLog() = 0;
+
+	/**
+	 * Gets a reference to the session that owns this instance.
+	 *
+	 * @return Owner session.
+	 */
+	virtual TSharedPtr<ISessionInfo> GetOwnerSession() = 0;
+
+	/**
+	 * Gets the name of the platform that the instance is running on.
+	 *
+	 * @return Platform name string.
+	 */
+	virtual const FString& GetPlatformName() const = 0;
+
+	/**
+	 * Gets the instance's current game world time.
+	 *
+	 * @return World time in seconds.
+	 */
+	virtual float GetWorldTimeSeconds() const = 0;
+
+	/**
+	 * Checks whether this instance is a console build (i.e. no editor features).
+	 *
+	 * @return true if it is a console build, false otherwise.
+	 */
+	virtual const bool IsConsole() const = 0;
+
+	/**
+	 * Checks whether this instance has already begun game play.
+	 *
+	 * @return true if game play has begun, false otherwise.
+	 */
+	virtual bool PlayHasBegun() const = 0;
+
+	/** Terminates the instance. */
+	virtual void Terminate() = 0;
+
+public:
+
+	/**
+	 * Gets an event delegate that is executed when a new log message has been received.
+	 *
+	 * @return The delegate.
+	 */
+	DECLARE_EVENT_TwoParams(ISessionInstanceInfo, FLogReceivedEvent, const TSharedRef<ISessionInstanceInfo>&, const TSharedRef<FSessionLogMessage>&);
+	virtual FLogReceivedEvent& OnLogReceived() = 0;
+
+public:
+
+	/** Virtual destructor. */
+	virtual ~ISessionInstanceInfo() { }
+};
+
+
+/** Type definition for shared pointers to instances of IGameInstanceInfo. */
+typedef TSharedPtr<ISessionInstanceInfo> ISessionInstanceInfoPtr;
+
+/** Type definition for shared references to instances of IGameInstanceInfo. */
+typedef TSharedRef<ISessionInstanceInfo> ISessionInstanceInfoRef;
diff --git a/Engine/Source/Developer/SessionServices/Public/ISessionManager.h b/Engine/Source/Developer/SessionServices/Public/ISessionManager.h
new file mode 100644
--- /dev/null
+++ b/Engine/Source/Developer/SessionServices/Public/ISessionManager.h
@@ -0,0 +1,135 @@
+// Copyright 1998-2015 Epic Games, Inc. All Rights Reserved.
+
+#pragma once
+
+
+class ISessionInfo;
+class ISessionInstanceInfo;
+
+
+/**
+ * Interface for the session manager.
+ */
+class ISessionManager
+{
+public:
+
+	/**
+	 * Adds an owner whose sessions we are interested in
+	 *
+	 * @param InOwner Session owner we want to view sessions from
+	 */
+	virtual void AddOwner( const FString& InOwner ) = 0;
+
+	/**
+	  * Gets the collection of currently selected engine instances.
+	  *
+	  * @param OutInstances Will hold the selected instances.
+	  */
+	virtual void GetSelectedInstances( TArray<TSharedPtr<ISessionInstanceInfo>>& OutInstances) const = 0;
+
+	/**
+	 * Get the selected session - as chosen in the session browser
+	 *
+	 * @return The session ID selected
+	 */
+	virtual const TSharedPtr<ISessionInfo>& GetSelectedSession() const = 0;
+
+	/**
+	 * Gets the list of all discovered sessions.
+	 *
+	 * @param OutSessions Will hold the collection of sessions.
+	 */
+	virtual void GetSessions( TArray<TSharedPtr<ISessionInfo>>& OutSessions ) const = 0;
+
+	/**
+	 * Checks whether the given instance is currently selected.
+	 *
+	 * @return true if the instance is selected, false otherwise.
+	 */
+	virtual bool IsInstanceSelected( const TSharedRef<ISessionInstanceInfo>& Instance ) const = 0;
+
+	/**
+	 * Removes an owner whose sessions we are no longer interested in
+	 *
+	 * @param InOwner Session owner we want to remove
+	 */
+	virtual void RemoveOwner( const FString& InOwner ) = 0;
+
+	/**
+	 * Selects the specified session.
+	 *
+	 * @param Session The session to the select (can be NULL to select none).
+	 * @return true if the session was selected, false otherwise.
+	 */
+	virtual bool SelectSession( const TSharedPtr<ISessionInfo>& Session ) = 0;
+
+	/**
+	 * Marks the specified item as selected or unselected.
+	 *
+	 * @param Instance The instance to mark.
+	 * @param Selected Whether the instance should be selected (true) or unselected (false).
+	 * @return true if the instance was selected, false otherwise.
+	 */
+	virtual bool SetInstanceSelected( const TSharedPtr<ISessionInstanceInfo>& Instance, bool Selected ) = 0;
+
+public:
+
+	/**
+	 * Returns a delegate that is executed before a session is being selected.
+	 *
+	 * @return The delegate.
+	 */
+	DECLARE_EVENT_TwoParams(ISessionManager, FCanSelectSessionEvent, const TSharedPtr<ISessionInfo>& /*Session*/, bool& /*OutCanSelect*/)
+	virtual FCanSelectSessionEvent& OnCanSelectSession() = 0;
+
+	/**
+	 * Returns a delegate that is executed when an instance changes its selection state.
+	 *
+	 * @return The delegate.
+	 */
+	DECLARE_EVENT(ISessionManager, FInstanceSelectionChangedEvent)
+	virtual FInstanceSelectionChangedEvent& OnInstanceSelectionChanged() = 0;
+
+	/**
+	 * Returns a delegate that is executed when the selected session received a log message from one of its instances.
+	 *
+	 * @return The delegate.
+	 */
+	DECLARE_EVENT_ThreeParams(ISessionManager, FLogReceivedEvent, const TSharedRef<ISessionInfo>& /*OwnerSession*/, const TSharedRef<ISessionInstanceInfo>& /*SessionInstance*/, const TSharedRef<FSessionLogMessage>& /*LogMessage*/);
+	virtual FLogReceivedEvent& OnLogReceived() = 0;
+
+	/**
+	 * Returns a delegate that is executed when the selected session changed.
+	 *
+	 * @return The delegate.
+	 */
+	DECLARE_EVENT_OneParam(ISessionManager, FSelectedSessionChangedEvent, const TSharedPtr<ISessionInfo>& /*SelectedSession*/)
+	virtual FSelectedSessionChangedEvent& OnSelectedSessionChanged() = 0;
+
+	/**
+	 * Returns a delegate that is executed when the list of sessions has changed.
+	 *
+	 * @return The delegate.
+	 */
+	virtual FSimpleMulticastDelegate& OnSessionsUpdated() = 0;
+
+	/**
+	 * Returns a delegate that is executed when a session instance is updated.
+	 *
+	 * @return The delegate.
+	 */
+	virtual FSimpleMulticastDelegate& OnSessionInstanceUpdated() = 0;
+
+public:
+
+	/** Virtual destructor. */
+	virtual ~ISessionManager() { }
+};
+
+
+/** Type definition for shared pointers to instances of ISessionManager. */
+typedef TSharedPtr<ISessionManager> ISessionManagerPtr;
+
+/** Type definition for shared references to instances of ISessionManager. */
+typedef TSharedRef<ISessionManager> ISessionManagerRef;
diff --git a/Engine/Source/Developer/SessionServices/Public/ISessionService.h b/Engine/Source/Developer/SessionServices/Public/ISessionService.h
new file mode 100644
--- /dev/null
+++ b/Engine/Source/Developer/SessionServices/Public/ISessionService.h
@@ -0,0 +1,47 @@
+// Copyright 1998-2015 Epic Games, Inc. All Rights Reserved.
+
+#pragma once
+
+
+/**
+ * Interface for application session services.
+ */
+class ISessionService
+{
+public:
+
+	/**
+	 * Checks whether the service is running.
+	 *
+	 * @return true if the service is running, false otherwise.
+	 * @see Start, Stop
+	 */
+	virtual bool IsRunning() = 0;
+
+	/**
+	 * Starts the service.
+	 *
+	 * @return true if the service was started, false otherwise.
+	 * @see IsRunning, Stop
+	 */
+	virtual bool Start() = 0;
+
+	/**
+	 * Stops the service.
+	 *
+	 * @see IsRunning, Start
+	 */
+	virtual void Stop() = 0;
+
+public:
+
+	/** Virtual destructor. */
+	virtual ~ISessionService() { }
+};
+
+
+/** Type definition for shared pointers to instances of ISessionService. */
+typedef TSharedPtr<ISessionService> ISessionServicePtr;
+
+/** Type definition for shared references to instances of ISessionService. */
+typedef TSharedRef<ISessionService> ISessionServiceRef;
diff --git a/Engine/Source/Developer/SessionServices/Public/ISessionServicesModule.h b/Engine/Source/Developer/SessionServices/Public/ISessionServicesModule.h
new file mode 100644
--- /dev/null
+++ b/Engine/Source/Developer/SessionServices/Public/ISessionServicesModule.h
@@ -0,0 +1,37 @@
+// Copyright 1998-2015 Epic Games, Inc. All Rights Reserved.
+
+#pragma once
+
+
+// forward declarations
+class ISessionManager;
+class ISessionService;
+
+
+/**
+ * Interface for session core modules.
+ */
+class ISessionServicesModule
+	: public IModuleInterface
+{
+public:
+
+	/**
+	 * Gets the session manager.
+	 *
+	 * @return The session manager.
+	 */
+	virtual TSharedRef<ISessionManager> GetSessionManager() = 0;
+
+	/** 
+	 * Gets the session service.
+	 *
+	 * @return The session service.
+	 */
+	virtual TSharedRef<ISessionService> GetSessionService() = 0;
+
+public:
+
+	/** Virtual destructor. */
+	virtual ~ISessionServicesModule() { }
+};
diff --git a/Engine/Source/Developer/SessionServices/Public/Interfaces/ISessionInfo.h b/Engine/Source/Developer/SessionServices/Public/Interfaces/ISessionInfo.h
deleted file mode 100644
diff --git a/Engine/Source/Developer/SessionServices/Public/Interfaces/ISessionInstanceInfo.h b/Engine/Source/Developer/SessionServices/Public/Interfaces/ISessionInstanceInfo.h
deleted file mode 100644
diff --git a/Engine/Source/Developer/SessionServices/Public/Interfaces/ISessionManager.h b/Engine/Source/Developer/SessionServices/Public/Interfaces/ISessionManager.h
deleted file mode 100644
diff --git a/Engine/Source/Developer/SessionServices/Public/Interfaces/ISessionService.h b/Engine/Source/Developer/SessionServices/Public/Interfaces/ISessionService.h
deleted file mode 100644
diff --git a/Engine/Source/Developer/SessionServices/Public/Interfaces/ISessionServicesModule.h b/Engine/Source/Developer/SessionServices/Public/Interfaces/ISessionServicesModule.h
deleted file mode 100644
diff --git a/Engine/Source/Developer/SessionServices/Public/Interfaces/SessionLogMessage.h b/Engine/Source/Developer/SessionServices/Public/Interfaces/SessionLogMessage.h
deleted file mode 100644
diff --git a/Engine/Source/Developer/SessionServices/Public/SessionLogMessage.h b/Engine/Source/Developer/SessionServices/Public/SessionLogMessage.h
new file mode 100644
--- /dev/null
+++ b/Engine/Source/Developer/SessionServices/Public/SessionLogMessage.h
@@ -0,0 +1,73 @@
+// Copyright 1998-2015 Epic Games, Inc. All Rights Reserved.
+
+#pragma once
+
+
+/**
+ * Structure for log messages.
+ */
+struct FSessionLogMessage
+{
+	/** Holds the log category. */
+	FName Category;
+
+	/** Holds the identifier of the engine instance that generated this log message. */
+	FGuid InstanceId;
+
+	/** Holds the name of the engine instance that generated this log message. */
+	FString InstanceName;
+
+	/** Holds the message text. */
+	FString Text;
+
+	/** Holds the time at which the message was generated. */
+	FDateTime Time;
+
+	/** Holds the number of seconds from the start of the instance at which the message was generated. */
+	double TimeSeconds;
+
+	/** Holds the verbosity type. */
+	ELogVerbosity::Type Verbosity;
+
+public:
+
+	/**
+	 * Creates and initializes a new instance.
+	 *
+	 * @param InInstanceId The identifier of the instance that generated the log message.
+	 * @param InInstanceName The name of the engine instance that generated the log message.
+	 * @param InTimeSeconds The number of seconds from the start of the instance at which the message was generated.
+	 * @param InText The message text.
+	 * @param InVerbosity The verbosity type.
+	 * @param InCategory The log category.
+	 */
+	FSessionLogMessage(const FGuid& InInstanceId, const FString& InInstanceName, float InTimeSeconds, const FString& InText, ELogVerbosity::Type InVerbosity, const FName& InCategory)
+		: Category(InCategory)
+		, InstanceId(InInstanceId)
+		, InstanceName(InInstanceName)
+		, Text((InCategory != NAME_None) ? InCategory.ToString() + TEXT(": ") + InText : InText)
+		, Time(FDateTime::Now())
+		, TimeSeconds(InTimeSeconds)
+		, Verbosity(InVerbosity)
+	{ }
+
+public:
+
+	/**
+	 * Implements a predicate to compare two log messages by log time.
+	 */
+	struct TimeComparer
+	{
+		bool operator()(const TSharedPtr<FSessionLogMessage>& A, const TSharedPtr<FSessionLogMessage>& B) const
+		{
+			return A->Time < B->Time;
+		}
+	};
+};
+
+
+/** Type definition for shared pointers to instances of FLauncherLogMessage. */
+typedef TSharedPtr<FSessionLogMessage> FSessionLogMessagePtr;
+
+/** Type definition for shared references to instances of FLauncherLogMessage. */
+typedef TSharedRef<FSessionLogMessage> FSessionLogMessageRef;
diff --git a/Engine/Source/Developer/SessionServices/Public/SessionServices.h b/Engine/Source/Developer/SessionServices/Public/SessionServices.h
deleted file mode 100644
diff --git a/Engine/Source/Runtime/Launch/Private/LaunchEngineLoop.cpp b/Engine/Source/Runtime/Launch/Private/LaunchEngineLoop.cpp
--- a/Engine/Source/Runtime/Launch/Private/LaunchEngineLoop.cpp
+++ b/Engine/Source/Runtime/Launch/Private/LaunchEngineLoop.cpp
@@ -46,6 +46,8 @@
 	#include "ContentStreaming.h"
 	#include "HighResScreenshot.h"
 	#include "HotReloadInterface.h"
+	#include "ISessionService.h"
+	#include "ISessionServicesModule.h"
 
 #if !UE_SERVER
 	#include "HeadMountedDisplay.h"
diff --git a/Engine/Source/Runtime/Launch/Private/LaunchPrivatePCH.h b/Engine/Source/Runtime/Launch/Private/LaunchPrivatePCH.h
--- a/Engine/Source/Runtime/Launch/Private/LaunchPrivatePCH.h
+++ b/Engine/Source/Runtime/Launch/Private/LaunchPrivatePCH.h
@@ -12,7 +12,6 @@
 #if WITH_ENGINE
 	#include "Engine.h"
 	#include "Messaging.h"
-	#include "SessionServices.h"
 	#include "SlateBasics.h"
 
 	#include "MallocProfilerEx.h"
diff --git a/Engine/Source/Runtime/Launch/Public/LaunchEngineLoop.h b/Engine/Source/Runtime/Launch/Public/LaunchEngineLoop.h
--- a/Engine/Source/Runtime/Launch/Public/LaunchEngineLoop.h
+++ b/Engine/Source/Runtime/Launch/Public/LaunchEngineLoop.h
@@ -2,11 +2,9 @@
 
 #pragma once
 
-#if WITH_ENGINE
-	#include "ISessionService.h"
-#endif
 
 class FEngineService;
+class ISessionService;
 
 
 /**
@@ -147,7 +145,7 @@ private:
 	FEngineService* EngineService;
 
 	/** Holds the application session service. */
-	ISessionServicePtr SessionService;
+	TSharedPtr<ISessionService> SessionService;
 
 #endif // WITH_ENGINE
 };
