Profiler - UE-10697 - Crash when opening Post Process Material with Live Data Preview running in Profiler (reimplemented with binary compatibility)

diff --git a/Engine/Source/Developer/ProfilerClient/Private/ProfilerClientManager.h b/Engine/Source/Developer/ProfilerClient/Private/ProfilerClientManager.h
--- a/Engine/Source/Developer/ProfilerClient/Private/ProfilerClientManager.h
+++ b/Engine/Source/Developer/ProfilerClient/Private/ProfilerClientManager.h
@@ -14,6 +14,7 @@ struct FServiceConnection
 		, ReadMessages(0)
 	{
 		MetaData.CriticalSection = &CriticalSection;
+		CurrentFrame = 0;
 	}
 
 	~FServiceConnection()
@@ -24,6 +25,7 @@ struct FServiceConnection
 	FServiceConnection(const FServiceConnection& InConnnection)
 	{
 		MetaData.CriticalSection = &CriticalSection;
+		CurrentFrame = 0;
 	}
 
 	/** Instance Id */
@@ -66,6 +68,9 @@ struct FServiceConnection
 	/** Current data loading progress. */
 	float DataLoadingProgress;
 
+	/** Current frame of data to process */
+	int64 CurrentFrame;
+
 	/** Current number of read messages. */
 	uint64 ReadMessages;
 
diff --git a/Engine/Source/Developer/ProfilerMessages/Classes/ProfilerServiceMessages.h b/Engine/Source/Developer/ProfilerMessages/Classes/ProfilerServiceMessages.h
--- a/Engine/Source/Developer/ProfilerMessages/Classes/ProfilerServiceMessages.h
+++ b/Engine/Source/Developer/ProfilerMessages/Classes/ProfilerServiceMessages.h
@@ -199,6 +199,10 @@ struct FProfilerServicePreviewAck
 	UPROPERTY()
 	FGuid InstanceId;
 
+	UPROPERTY()
+	int64 Frame;
+
+
 	/**
 	 * Default constructor.
 	 */
@@ -207,8 +211,9 @@ struct FProfilerServicePreviewAck
 	/**
 	 * Creates and initializes a new instance.
 	 */
-	FProfilerServicePreviewAck( const FGuid& InInstance )
+	FProfilerServicePreviewAck( const FGuid& InInstance, int64 InFrame )
 		: InstanceId(InInstance)
+		, Frame(InFrame)
 	{
 	}
 };
diff --git a/Engine/Source/Developer/ProfilerService/Private/ProfilerServiceManager.cpp b/Engine/Source/Developer/ProfilerService/Private/ProfilerServiceManager.cpp
--- a/Engine/Source/Developer/ProfilerService/Private/ProfilerServiceManager.cpp
+++ b/Engine/Source/Developer/ProfilerService/Private/ProfilerServiceManager.cpp
@@ -496,7 +496,7 @@ void FProfilerServiceManager::SetPreviewState( const FMessageAddress& ClientAddr
 				Client->Preview = true;
 				if (MessageEndpoint.IsValid())
 				{
-					MessageEndpoint->Send( new FProfilerServicePreviewAck( InstanceId ), ClientAddress );
+					MessageEndpoint->Send( new FProfilerServicePreviewAck( InstanceId, 0 ), ClientAddress );
 				}
 				SendMetaData(ClientAddress);
 			}
diff --git a/Engine/Source/Developer/ProfilerService/Private/ProfilerServiceManager.h b/Engine/Source/Developer/ProfilerService/Private/ProfilerServiceManager.h
--- a/Engine/Source/Developer/ProfilerService/Private/ProfilerServiceManager.h
+++ b/Engine/Source/Developer/ProfilerService/Private/ProfilerServiceManager.h
@@ -16,6 +16,9 @@ struct FClientData
 #if STATS
 	FStatsWriteFile StatsWriteFile;
 
+	/** Current frame of data sent */
+	int64 CurrentFrame;
+
 	/** Stats metadata size. */
 	int32 MetadataSize;
 #endif
