Core - UE-12192 / Crash previewing capture data in session frontend (removed a bunch of dead code, still WIP)

diff --git a/Engine/Source/Developer/ProfilerClient/Private/ProfilerClientManager.cpp b/Engine/Source/Developer/ProfilerClient/Private/ProfilerClientManager.cpp
--- a/Engine/Source/Developer/ProfilerClient/Private/ProfilerClientManager.cpp
+++ b/Engine/Source/Developer/ProfilerClient/Private/ProfilerClientManager.cpp
@@ -31,11 +31,9 @@ FProfilerClientManager::FProfilerClientManager( const IMessageBusRef& InMessageB
 #if STATS
 	MessageBus = InMessageBus;
 	MessageEndpoint = FMessageEndpoint::Builder("FProfilerClientModule", InMessageBus)
-		.Handling<FProfilerServiceAuthorize>(this, &FProfilerClientManager::HandleServiceAuthorizeMessage)
 		.Handling<FProfilerServiceAuthorize2>(this, &FProfilerClientManager::HandleServiceAuthorize2Message)
 		.Handling<FProfilerServiceData2>(this, &FProfilerClientManager::HandleServiceData2Message)
 		.Handling<FProfilerServicePreviewAck>(this, &FProfilerClientManager::HandleServicePreviewAckMessage)
-		.Handling<FProfilerServiceMetaData>(this, &FProfilerClientManager::HandleServiceMetaDataMessage)
 		.Handling<FProfilerServiceFileChunk>(this, &FProfilerClientManager::HandleServiceFileChunk)  
 		.Handling<FProfilerServicePing>(this, &FProfilerClientManager::HandleServicePingMessage);
 
@@ -309,21 +307,6 @@ void FProfilerClientManager::LoadCapture( const FString& DataFilepath, const FGu
 }
 
 
-void FProfilerClientManager::RequestMetaData()
-{
-#if STATS
-	if (MessageEndpoint.IsValid() && ActiveSessionId.IsValid())
-	{
-		TArray<FMessageAddress> Instances;
-		for (auto It = Connections.CreateConstIterator(); It; ++It)
-		{
-			Instances.Add(It.Value().ServiceAddress);
-		}
-		MessageEndpoint->Send(new FProfilerServiceRequest(EProfilerRequestType::PRT_MetaData), Instances);
-	}
-#endif
-}
-
 void FProfilerClientManager::RequestLastCapturedFile( const FGuid& InstanceId /*= FGuid()*/ )
 {
 #if STATS
@@ -363,25 +346,6 @@ void FProfilerClientManager::HandleMessageBusShutdown()
 #endif
 }
 
-
-void FProfilerClientManager::HandleServiceAuthorizeMessage( const FProfilerServiceAuthorize& Message, const IMessageContextRef& Context )
-{
-#if STATS
-	if (ActiveSessionId == Message.SessionId && PendingInstances.Contains(Message.InstanceId))
-	{
-		PendingInstances.Remove(Message.InstanceId);
-		FServiceConnection& Connection = Connections.FindOrAdd(Message.InstanceId);
-		Connection.ServiceAddress = Context->GetSender();
-		Connection.InstanceId = Message.InstanceId;
-		Connection.CurrentData.Frame = 0;
-	}
-
-	// Fire the client connection event
-	ProfilerClientConnectedDelegate.Broadcast(ActiveSessionId, Message.InstanceId);
-#endif
-}
-
-
 void FProfilerClientManager::HandleServiceAuthorize2Message( const FProfilerServiceAuthorize2& Message, const IMessageContextRef& Context )
 {
 #if STATS
@@ -440,41 +404,6 @@ void FServiceConnection::Initialize( const FProfilerServiceAuthorize2& Message,
 #endif
 }
 
-void FProfilerClientManager::HandleServiceMetaDataMessage( const FProfilerServiceMetaData& Message, const IMessageContextRef& Context )
-{
-#if STATS
-	if (ActiveSessionId.IsValid() && Connections.Find(Message.InstanceId) != nullptr)
-	{
-		FServiceConnection& Connection = *Connections.Find(Message.InstanceId);
-
-		FArrayReader ArrayReader(true);
-		ArrayReader.Append(Message.Data);
-
-		FStatMetaData NewData;
-		ArrayReader << NewData;
-
-		// add in the new data
-		if (NewData.StatDescriptions.Num() > 0)
-		{
-			Connection.MetaData.StatDescriptions.Append(NewData.StatDescriptions);
-		}
-		if (NewData.GroupDescriptions.Num() > 0)
-		{
-			Connection.MetaData.GroupDescriptions.Append(NewData.GroupDescriptions);
-		}
-		if (NewData.ThreadDescriptions.Num() > 0)
-		{
-			Connection.MetaData.ThreadDescriptions.Append(NewData.ThreadDescriptions);
-		}
-		Connection.MetaData.SecondsPerCycle = NewData.SecondsPerCycle;
-
-		// Fire a meta data update message
-		ProfilerMetaDataUpdatedDelegate.Broadcast(Message.InstanceId);
-	}
-#endif
-}
-
-
 bool FProfilerClientManager::CheckHashAndWrite( const FProfilerServiceFileChunk& FileChunk, const FProfilerFileChunkHeader& FileChunkHeader, FArchive* Writer )
 {
 #if STATS
@@ -748,6 +677,7 @@ void FProfilerClientManager::HandleServicePreviewAckMessage( const FProfilerServ
 	if (ActiveSessionId.IsValid() && Connections.Find(Message.InstanceId) != NULL)
 	{
 		FServiceConnection& Connection = *Connections.Find(Message.InstanceId);
+		Connection.CurrentFrame = Message.Frame;//
 	}
 #endif
 }
diff --git a/Engine/Source/Developer/ProfilerClient/Private/ProfilerClientManager.h b/Engine/Source/Developer/ProfilerClient/Private/ProfilerClientManager.h
--- a/Engine/Source/Developer/ProfilerClient/Private/ProfilerClientManager.h
+++ b/Engine/Source/Developer/ProfilerClient/Private/ProfilerClientManager.h
@@ -165,7 +165,6 @@ public:
 	virtual void SetCaptureState( const bool bRequestedCaptureState, const FGuid& InstanceId = FGuid() ) override;
 	virtual void SetPreviewState( const bool bRequestedPreviewState, const FGuid& InstanceId = FGuid() ) override;
 	virtual void LoadCapture( const FString& DataFilepath, const FGuid& ProfileId ) override;
-	virtual void RequestMetaData() override;
 	virtual void RequestLastCapturedFile( const FGuid& InstanceId = FGuid() ) override;
 
 	virtual const FStatMetaData& GetStatMetaData( const FGuid& InstanceId ) const override
@@ -221,15 +220,9 @@ private:
 	// Handles message bus shutdowns.
 	void HandleMessageBusShutdown();
 
-	// Handles FProfilerServiceAuthorize messages.
-	void HandleServiceAuthorizeMessage( const FProfilerServiceAuthorize& Message, const IMessageContextRef& Context );
-
 	// Handles FProfilerServiceAuthorize2 messages.
 	void HandleServiceAuthorize2Message( const FProfilerServiceAuthorize2& Message, const IMessageContextRef& Context );
 
-	// Handles FProfilerServiceMetaData messages.
-	void HandleServiceMetaDataMessage( const FProfilerServiceMetaData& Message, const IMessageContextRef& Context );
-
 	// Handles FProfilerServiceFileChunk messages.
 	void HandleServiceFileChunk( const FProfilerServiceFileChunk& FileChunk, const IMessageContextRef& Context );
 
diff --git a/Engine/Source/Developer/ProfilerClient/Public/Interfaces/IProfilerClient.h b/Engine/Source/Developer/ProfilerClient/Public/Interfaces/IProfilerClient.h
--- a/Engine/Source/Developer/ProfilerClient/Public/Interfaces/IProfilerClient.h
+++ b/Engine/Source/Developer/ProfilerClient/Public/Interfaces/IProfilerClient.h
@@ -84,9 +84,6 @@ public:
 	 */
 	virtual void LoadCapture( const FString& DataFilepath, const FGuid& ProfileId ) = 0;
 
-	/** Requests the meta data from the service */
-	virtual void RequestMetaData() = 0;
-
 	/**
 	 * Requests the last captured file from the service.
 	 *
diff --git a/Engine/Source/Developer/ProfilerMessages/Classes/ProfilerServiceMessages.h b/Engine/Source/Developer/ProfilerMessages/Classes/ProfilerServiceMessages.h
--- a/Engine/Source/Developer/ProfilerMessages/Classes/ProfilerServiceMessages.h
+++ b/Engine/Source/Developer/ProfilerMessages/Classes/ProfilerServiceMessages.h
@@ -183,39 +183,6 @@ struct FProfilerServicePreviewAck
 };
 
 
-/**
- */
-USTRUCT()
-struct FProfilerServiceMetaData
-{
-	GENERATED_USTRUCT_BODY()
-
-	/**
-	 */
-	UPROPERTY()
-	FGuid InstanceId;
-
-	/**
-	 */
-	UPROPERTY()
-	TArray<uint8> Data;
-
-
-	/**
-	 * Default constructor.
-	 */
-	FProfilerServiceMetaData( ) { }
-
-	/**
-	 * Creates and initializes a new instance.
-	 */
-	FProfilerServiceMetaData( const FGuid& InInstance, const TArray<uint8>& InData )
-		: InstanceId(InInstance)
-		, Data( InData )
-	{}
-};
-
-
 /**
  * Implements a message for copying a file through the network, as well as for synchronization.
  * Unfortunately assumes that InstanceId and Filename are transfered without errors.
diff --git a/Engine/Source/Developer/ProfilerService/Private/ProfilerServiceManager.cpp b/Engine/Source/Developer/ProfilerService/Private/ProfilerServiceManager.cpp
--- a/Engine/Source/Developer/ProfilerService/Private/ProfilerServiceManager.cpp
+++ b/Engine/Source/Developer/ProfilerService/Private/ProfilerServiceManager.cpp
@@ -316,41 +316,9 @@ FProfilerServiceManager::FProfilerServiceManager()
 	DataFrame.Frame = 0;
 }
 
-
-FProfilerServiceManager::~FProfilerServiceManager()
-{}
-
-
 /* IProfilerServiceManager interface
  *****************************************************************************/
 
-void FProfilerServiceManager::SendData(FProfilerCycleCounter& Data)
-{
-	// add the data to the data frame
-	DataFrame.CycleCounters.FindOrAdd(Data.ThreadId).Add(Data);
-}
-
-
-void FProfilerServiceManager::SendData(FProfilerFloatAccumulator& Data)
-{
-	// add the data to the data frame
-	DataFrame.FloatAccumulators.Add(Data);
-}
-
-
-void FProfilerServiceManager::SendData(FProfilerCountAccumulator& Data)
-{
-	// add the data to the data frame
-	DataFrame.CountAccumulators.Add(Data);
-}
-
-
-void FProfilerServiceManager::SendData(FProfilerCycleGraph& Data)
-{
-	// add the data to the data frame
-	DataFrame.CycleGraphs.FindOrAdd(Data.ThreadId) = Data;
-}
-
 void FProfilerServiceManager::StartCapture()
 {
 #if STATS
@@ -368,40 +336,6 @@ void FProfilerServiceManager::StopCapture()
 #endif
 }
 
-void FProfilerServiceManager::SendMetaData(const FMessageAddress& client)
-{
-	if (MessageEndpoint.IsValid())
-	{
-		FArrayWriter ArrayWriter(true);
-		ArrayWriter << MetaData;
-		MessageEndpoint->Send(new FProfilerServiceMetaData(InstanceId, ArrayWriter), client);
-	}
-}
-
-void FProfilerServiceManager::StartFrame(uint32 FrameNumber, double FrameStart)
-{
-	// send data to the clients
-	if (DataFrame.Frame > 0)
-	{
-		if (MessageEndpoint.IsValid() && PreviewClients.Num() > 0)
-		{
-			FArrayWriter ArrayWriter(true);
-			ArrayWriter << DataFrame;
-			MessageEndpoint->Send(new FProfilerServiceData(InstanceId, ArrayWriter), PreviewClients );
-		}
-
-		ProfilerDataDelegate.Broadcast(FGuid(), DataFrame);
-	}
-
-	// update the data frame
-	DataFrame.CycleCounters.Reset();
-	DataFrame.CycleGraphs.Reset();
-	DataFrame.CountAccumulators.Reset();
-	DataFrame.FloatAccumulators.Reset();
-	DataFrame.Frame = FrameNumber;
-	DataFrame.FrameStart = FrameStart;
-}
-
 /* FProfilerServiceManager implementation
  *****************************************************************************/
 
@@ -496,9 +430,9 @@ void FProfilerServiceManager::SetPreviewState( const FMessageAddress& ClientAddr
 				Client->Preview = true;
 				if (MessageEndpoint.IsValid())
 				{
-					MessageEndpoint->Send( new FProfilerServicePreviewAck( InstanceId, 0 ), ClientAddress );
+					Client->CurrentFrame = FStats::GameThreadStatsFrame;
+					MessageEndpoint->Send( new FProfilerServicePreviewAck( InstanceId, FStats::GameThreadStatsFrame ), ClientAddress );
 				}
-				SendMetaData(ClientAddress);
 			}
 			else
 			{
@@ -602,11 +536,7 @@ void FProfilerServiceManager::HandleServicePreviewMessage( const FProfilerServic
 
 void FProfilerServiceManager::HandleServiceRequestMessage( const FProfilerServiceRequest& Message, const IMessageContextRef& Context )
 {
-	if (Message.Request == EProfilerRequestType::PRT_MetaData)
-	{
-		SendMetaData(Context->GetSender());
-	}
-	else if( Message.Request == EProfilerRequestType::PRT_SendLastCapturedFile )
+	if( Message.Request == EProfilerRequestType::PRT_SendLastCapturedFile )
 	{
 		if( LastStatsFilename.IsEmpty() == false )
 		{
diff --git a/Engine/Source/Developer/ProfilerService/Private/ProfilerServiceManager.h b/Engine/Source/Developer/ProfilerService/Private/ProfilerServiceManager.h
--- a/Engine/Source/Developer/ProfilerService/Private/ProfilerServiceManager.h
+++ b/Engine/Source/Developer/ProfilerService/Private/ProfilerServiceManager.h
@@ -39,28 +39,14 @@ public:
 	 */
 	FProfilerServiceManager();
 
-	/**
-	 * Default destructor
-	 */
-	~FProfilerServiceManager();
-
 public:
 
 	// Begin IProfilerServiceManager interface
-	virtual void SendData(FProfilerCycleCounter& Data) override;
-
-	virtual void SendData(FProfilerFloatAccumulator& Data) override;
-
-	virtual void SendData(FProfilerCountAccumulator& Data) override;
-
-	virtual void SendData(FProfilerCycleGraph& Data) override;
 
 	virtual void StartCapture() override;
 	
 	virtual void StopCapture() override;
 
-	virtual void StartFrame(uint32 FrameNumber, double FrameStart) override;
-
 	virtual FStatMetaData& GetStatMetaData() override
 	{
 		return MetaData;
@@ -97,8 +83,6 @@ private:
 	/** Callback for a tick, used to ping the clients */
 	bool HandlePing( float DeltaTime );
 
-	void SendMetaData(const FMessageAddress& client);
-
 private:
 
 	// Handles FProfilerServiceCapture messages.
diff --git a/Engine/Source/Developer/ProfilerService/Public/Interfaces/IProfilerServiceManager.h b/Engine/Source/Developer/ProfilerService/Public/Interfaces/IProfilerServiceManager.h
--- a/Engine/Source/Developer/ProfilerService/Public/Interfaces/IProfilerServiceManager.h
+++ b/Engine/Source/Developer/ProfilerService/Public/Interfaces/IProfilerServiceManager.h
@@ -366,8 +366,6 @@ namespace EProfilerRequestType
 {
 	enum Type
 	{
-		/** Metadata request. */
-		PRT_MetaData,
 		/** Send last captured file. */
 		PRT_SendLastCapturedFile,
 	};
@@ -473,43 +471,12 @@ class IProfilerServiceManager
 {
 public:
 
-	/**
-	 * Sends profiler cycle counter to registered clients.
-	 *
-	 * @param Data The profiler data to be sent.
-	 */
-	virtual void SendData(FProfilerCycleCounter& Data) = 0;
-
-	/**
-	 * Sends profiler float accumulator to registered clients.
-	 *
-	 * @param Data The profiler data to be sent.
-	 */
-	virtual void SendData(FProfilerFloatAccumulator& Data) = 0;
-
-	/**
-	 * Sends profiler count accumulator to registered clients.
-	 *
-	 * @param Data The profiler data to be sent.
-	 */
-	virtual void SendData(FProfilerCountAccumulator& Data) = 0;
-
-	/**
-	 * Sends profiler cycle graph to registered clients.
-	 *
-	 * @param Data The profiler data to be sent.
-	 */
-	virtual void SendData(FProfilerCycleGraph& Data) = 0;
-
 	/** Starts a file capture. */
 	virtual void StartCapture() = 0;
 
 	/** Stops a file capture. */
 	virtual void StopCapture() = 0;
 
-	/** Starts a new frame of data and sends the previous if it exists. */
-	virtual void StartFrame(uint32 FrameNumber, double FrameStart) = 0;
-
 	/**
 	 * Gets the description for the given stat id.
 	 *
