Core - UE-12192/Crash previewing capture data in session frontend

diff --git a/Engine/Source/Developer/ProfilerClient/Private/ProfilerClientManager.cpp b/Engine/Source/Developer/ProfilerClient/Private/ProfilerClientManager.cpp
--- a/Engine/Source/Developer/ProfilerClient/Private/ProfilerClientManager.cpp
+++ b/Engine/Source/Developer/ProfilerClient/Private/ProfilerClientManager.cpp
@@ -10,7 +10,7 @@
 #include "StatsData.h"
 #include "StatsFile.h"
 
-DEFINE_LOG_CATEGORY_STATIC(LogProfile, Log, All);
+DEFINE_LOG_CATEGORY_STATIC( LogProfilerClient, Log, All );
 
 DECLARE_CYCLE_STAT( TEXT("HandleDataReceived"),	STAT_PC_HandleDataReceived,			STATGROUP_Profiler );
 DECLARE_CYCLE_STAT( TEXT("ReadStatMessages"),	STAT_PC_ReadStatMessages,			STATGROUP_Profiler );
@@ -31,15 +31,15 @@ FProfilerClientManager::FProfilerClientManager( const IMessageBusRef& InMessageB
 #if STATS
 	MessageBus = InMessageBus;
 	MessageEndpoint = FMessageEndpoint::Builder("FProfilerClientModule", InMessageBus)
-		.Handling<FProfilerServiceAuthorize2>(this, &FProfilerClientManager::HandleServiceAuthorize2Message)
-		.Handling<FProfilerServiceData2>(this, &FProfilerClientManager::HandleServiceData2Message)
+		.Handling<FProfilerServiceAuthorize>(this, &FProfilerClientManager::HandleServiceAuthorizeMessage)
+		.Handling<FProfilerServiceData2>(this, &FProfilerClientManager::HandleProfilerServiceData2Message)
 		.Handling<FProfilerServicePreviewAck>(this, &FProfilerClientManager::HandleServicePreviewAckMessage)
 		.Handling<FProfilerServiceFileChunk>(this, &FProfilerClientManager::HandleServiceFileChunk)  
 		.Handling<FProfilerServicePing>(this, &FProfilerClientManager::HandleServicePingMessage);
 
 	if (MessageEndpoint.IsValid())
 	{
-		InMessageBus->OnShutdown().AddRaw(this, &FProfilerClientManager::HandleMessageBusShutdown);
+		OnShutdownMessageBusDelegateHandle = InMessageBus->OnShutdown().AddRaw( this, &FProfilerClientManager::HandleMessageBusShutdown );
 		MessageEndpoint->Subscribe<FProfilerServicePing>();
 	}
 
@@ -52,7 +52,7 @@ FProfilerClientManager::FProfilerClientManager( const IMessageBusRef& InMessageB
 	LoadTask = nullptr;
 #endif
 	LoadConnection = nullptr;
-	MessageDelegateHandle = FTicker::GetCoreTicker().AddTicker(MessageDelegate, 0.05f);
+	MessageDelegateHandle = FTicker::GetCoreTicker().AddTicker(MessageDelegate, 0.1f);
 #endif
 }
 
@@ -69,7 +69,7 @@ FProfilerClientManager::~FProfilerClientManager()
 
 		IFileManager::Get().Delete( *ReceivedFileInfo.DestFilepath );
 
-		UE_LOG(LogProfile, Log, TEXT( "File service-client transfer aborted: %s" ), *It.Key() );
+		UE_LOG( LogProfilerClient, Log, TEXT( "File service-client transfer aborted: %s" ), *It.Key() );
 	}
 
 	FTicker::GetCoreTicker().RemoveTicker(MessageDelegateHandle);
@@ -79,7 +79,7 @@ FProfilerClientManager::~FProfilerClientManager()
 
 	if (MessageBus.IsValid())
 	{
-		MessageBus->OnShutdown().RemoveAll(this);
+		MessageBus->OnShutdown().Remove( OnShutdownMessageBusDelegateHandle );
 	}
 
 	LoadConnection = nullptr;
@@ -115,6 +115,8 @@ void FProfilerClientManager::Subscribe( const FGuid& Session )
 	}
 
 	Connections.Reset();
+
+	UE_LOG( LogProfilerClient, Log, TEXT( "Subscribe Session: %s" ), *Session.ToString() );
 #endif
 }
 
@@ -129,6 +131,8 @@ void FProfilerClientManager::Track( const FGuid& Instance )
 
 		RetryTime = 5.f;
 		TickDelegateHandle = FTicker::GetCoreTicker().AddTicker(TickDelegate, RetryTime);
+
+		UE_LOG( LogProfilerClient, Verbose, TEXT( "Track Session: %s, Instance: %s" ), *ActiveSessionId.ToString(), *Instance.ToString() );
 	}
 #endif
 }
@@ -171,6 +175,8 @@ void FProfilerClientManager::Untrack( const FGuid& Instance )
 
 		// fire the disconnection delegate
 		ProfilerClientDisconnectedDelegate.Broadcast(ActiveSessionId, Instance);
+
+		UE_LOG( LogProfilerClient, Verbose, TEXT( "Untrack Session: %s, Instance: %s" ), *ActiveSessionId.ToString(), *Instance.ToString() );
 	}
 #endif
 }
@@ -193,17 +199,20 @@ void FProfilerClientManager::SetCaptureState( const bool bRequestedCaptureState,
 			TArray<FMessageAddress> Instances;
 			for (auto It = Connections.CreateConstIterator(); It; ++It)
 			{
-				Instances.Add(It.Value().ServiceAddress);
+				Instances.Add( It.Value().ProfilerServiceAddress );
 			}
 			MessageEndpoint->Send(new FProfilerServiceCapture(bRequestedCaptureState), Instances);
+			UE_LOG( LogProfilerClient, Verbose, TEXT( "SetCaptureState Session: %s, Instance: %s, State: %i" ), *ActiveSessionId.ToString(), *InstanceId.ToString(), (int32)bRequestedCaptureState );
 		}
 		else
 		{
-			const FMessageAddress* MessageAddress = &Connections.Find(InstanceId)->ServiceAddress;
+			const FMessageAddress* MessageAddress = &Connections.Find(InstanceId)->ProfilerServiceAddress;
 			if( MessageAddress )
 			{
 				MessageEndpoint->Send(new FProfilerServiceCapture(bRequestedCaptureState), *MessageAddress);
 			}
+
+			UE_LOG( LogProfilerClient, Verbose, TEXT( "SetCaptureState Session: %s, Instance: %s, State: %i" ), *ActiveSessionId.ToString(), *InstanceId.ToString(), (int32)bRequestedCaptureState );
 		}
 	}
 #endif
@@ -219,17 +228,19 @@ void FProfilerClientManager::SetPreviewState( const bool bRequestedPreviewState,
 			TArray<FMessageAddress> Instances;
 			for (auto It = Connections.CreateConstIterator(); It; ++It)
 			{
-				Instances.Add(It.Value().ServiceAddress);
+				Instances.Add( It.Value().ProfilerServiceAddress );
 			}
 			MessageEndpoint->Send(new FProfilerServicePreview(bRequestedPreviewState), Instances);
+			UE_LOG( LogProfilerClient, Verbose, TEXT( "SetPreviewState Session: %s, Instance: %s, State: %i" ), *ActiveSessionId.ToString(), *InstanceId.ToString(), (int32)bRequestedPreviewState );
 		}
 		else
 		{
-			const FMessageAddress* MessageAddress = &Connections.Find(InstanceId)->ServiceAddress;
+			const FMessageAddress* MessageAddress = &Connections.Find(InstanceId)->ProfilerServiceAddress;
 			if( MessageAddress )
 			{
 				MessageEndpoint->Send(new FProfilerServicePreview(bRequestedPreviewState), *MessageAddress);
 			}
+			UE_LOG( LogProfilerClient, Verbose, TEXT( "SetPreviewState Session: %s, Instance: %s, State: %i" ), *ActiveSessionId.ToString(), *InstanceId.ToString(), (int32)bRequestedPreviewState );
 		}
 	}
 #endif
@@ -247,7 +258,7 @@ void FProfilerClientManager::LoadCapture( const FString& DataFilepath, const FGu
 	const int64 Size = IFileManager::Get().FileSize( *DataFilepath );
 	if( Size < 4 )
 	{
-		UE_LOG( LogProfile, Error, TEXT( "Could not open: %s" ), *DataFilepath );
+		UE_LOG( LogProfilerClient, Error, TEXT( "Could not open: %s" ), *DataFilepath );
 		return;
 	}
 
@@ -259,13 +270,13 @@ void FProfilerClientManager::LoadCapture( const FString& DataFilepath, const FGu
 
 	if( !FileReader )
 	{
-		UE_LOG( LogProfile, Error, TEXT( "Could not open: %s" ), *DataFilepath );
+		UE_LOG( LogProfilerClient, Error, TEXT( "Could not open: %s" ), *DataFilepath );
 		return;
 	}
 
 	if( !LoadConnection->Stream.ReadHeader( *FileReader ) )
 	{
-		UE_LOG( LogProfile, Error, TEXT( "Could not open, bad magic: %s" ), *DataFilepath );
+		UE_LOG( LogProfilerClient, Error, TEXT( "Could not open, bad magic: %s" ), *DataFilepath );
 		delete FileReader;
 		return;
 	}
@@ -292,7 +303,7 @@ void FProfilerClientManager::LoadCapture( const FString& DataFilepath, const FGu
 
 	if( LoadConnection->Stream.Header.HasCompressedData() )
 	{
-		UE_CLOG( !bIsFinalized, LogProfile, Fatal, TEXT( "Compressed stats file has to be finalized" ) );
+		UE_CLOG( !bIsFinalized, LogProfilerClient, Fatal, TEXT( "Compressed stats file has to be finalized" ) );
 	}
 
 #if PROFILER_THREADED_LOAD
@@ -317,13 +328,13 @@ void FProfilerClientManager::RequestLastCapturedFile( const FGuid& InstanceId /*
 			TArray<FMessageAddress> Instances;
 			for (auto It = Connections.CreateConstIterator(); It; ++It)
 			{
-				Instances.Add(It.Value().ServiceAddress);
+				Instances.Add( It.Value().ProfilerServiceAddress );
 			}
 			MessageEndpoint->Send(new FProfilerServiceRequest(EProfilerRequestType::PRT_SendLastCapturedFile), Instances);
 		}
 		else
 		{
-			const FMessageAddress* MessageAddress = &Connections.Find(InstanceId)->ServiceAddress;
+			const FMessageAddress* MessageAddress = &Connections.Find(InstanceId)->ProfilerServiceAddress;
 			if( MessageAddress )
 			{
 				MessageEndpoint->Send(new FProfilerServiceRequest(EProfilerRequestType::PRT_SendLastCapturedFile), *MessageAddress);
@@ -346,7 +357,7 @@ void FProfilerClientManager::HandleMessageBusShutdown()
 #endif
 }
 
-void FProfilerClientManager::HandleServiceAuthorize2Message( const FProfilerServiceAuthorize2& Message, const IMessageContextRef& Context )
+void FProfilerClientManager::HandleServiceAuthorizeMessage( const FProfilerServiceAuthorize& Message, const IMessageContextRef& Context )
 {
 #if STATS
 	if (ActiveSessionId == Message.SessionId && PendingInstances.Contains(Message.InstanceId))
@@ -360,35 +371,36 @@ void FProfilerClientManager::HandleServiceAuthorize2Message( const FProfilerServ
 
 		// Fire the client connection event
 		ProfilerClientConnectedDelegate.Broadcast(ActiveSessionId, Message.InstanceId);
+
+		UE_LOG( LogProfilerClient, Verbose, TEXT( "Authorize SessionId: %s, InstanceId: %s" ), *Message.SessionId.ToString(), *Message.InstanceId.ToString() );
 	}
 #endif
 }
 
-void FServiceConnection::Initialize( const FProfilerServiceAuthorize2& Message, const IMessageContextRef& Context )
+void FServiceConnection::Initialize( const FProfilerServiceAuthorize& Message, const IMessageContextRef& Context )
 {
 #if STATS
-	ServiceAddress = Context->GetSender();
+	ProfilerServiceAddress = Context->GetSender();
 	InstanceId = Message.InstanceId;
 	CurrentData.Frame = 0;
 
-	// add the supplied meta data
-	FArrayReader ArrayReader(true);//
-	ArrayReader.Append(Message.Data);
+	// Add the supplied meta data.
+	FMemoryReader MemoryReader( Message.Data );
 
 	MetaData.CriticalSection = &CriticalSection;
-	int64 Size = ArrayReader.TotalSize();
+	int64 Size = MemoryReader.TotalSize();
 
-	const bool bVerifyHeader = Stream.ReadHeader( ArrayReader );
+	const bool bVerifyHeader = Stream.ReadHeader( MemoryReader );
 	check( bVerifyHeader );
 
-	// read in the data
+	// Read in the data.
 	TArray<FStatMessage> StatMessages;
 	{
 		SCOPE_CYCLE_COUNTER(STAT_PC_ReadStatMessages);
-		while(ArrayReader.Tell() < Size)
+		while(MemoryReader.Tell() < Size)
 		{
 			// read the message
-			new (StatMessages) FStatMessage(Stream.ReadMessage(ArrayReader));
+			new (StatMessages)FStatMessage( Stream.ReadMessage( MemoryReader ) );
 		}
 		static FStatNameAndInfo Adv(NAME_AdvanceFrame, "", "", TEXT(""), EStatDataType::ST_int64, true, false);
 		new (StatMessages) FStatMessage(Adv.GetEncodedName(), EStatOperation::AdvanceFrameEventGameThread, 1LL, false);
@@ -443,7 +455,7 @@ void FProfilerClientManager::HandleServiceFileChunk( const FProfilerServiceFileC
 	const TCHAR* StrTmp = TEXT(".tmp");
 
 	// Read file chunk header.
-	FMemoryReader Reader(FileChunk.Header);
+	FMemoryReader Reader( FileChunk.Header );
 	FProfilerFileChunkHeader FileChunkHeader;
 	Reader << FileChunkHeader;
 	FileChunkHeader.Validate();
@@ -459,12 +471,12 @@ void FProfilerClientManager::HandleServiceFileChunk( const FProfilerServiceFileC
 			const FString PathName = FPaths::ProfilingDir() + TEXT("UnrealStats/Received/");
 			const FString StatFilepath = PathName + FileChunk.Filename + StrTmp;
 
-			UE_LOG(LogProfile, Log, TEXT( "Opening stats file for service-client sending: %s" ), *StatFilepath );
+			UE_LOG( LogProfilerClient, Log, TEXT( "Opening stats file for service-client sending: %s" ), *StatFilepath );
 
 			FArchive* FileWriter = IFileManager::Get().CreateFileWriter(*StatFilepath);
 			if( !FileWriter )
 			{
-				UE_LOG(LogProfile, Error, TEXT( "Could not open: %s" ), *StatFilepath );
+				UE_LOG( LogProfilerClient, Error, TEXT( "Could not open: %s" ), *StatFilepath );
 				return;
 			}
 
@@ -497,7 +509,7 @@ void FProfilerClientManager::HandleServiceFileChunk( const FProfilerServiceFileC
 
 				ActiveTransfers.Remove( FileChunk.Filename );
 
-				UE_LOG(LogProfile, Log, TEXT( "File service-client received successfully: %s" ), *FileChunk.Filename );
+				UE_LOG( LogProfilerClient, Log, TEXT( "File service-client received successfully: %s" ), *FileChunk.Filename );
 			}
 			else
 			{
@@ -510,7 +522,7 @@ void FProfilerClientManager::HandleServiceFileChunk( const FProfilerServiceFileC
 			if( MessageEndpoint.IsValid() )
 			{
 				MessageEndpoint->Send( new FProfilerServiceFileChunk(FileChunk,FProfilerServiceFileChunk::FNullTag()), Context->GetSender() );
-				UE_LOG(LogProfile, Log, TEXT("Received a bad chunk of file, resending: %5i, %6u, %10u, %s"), FileChunk.Data.Num(), ReceivedFileInfo->Progress, FileChunkHeader.FileSize, *FileChunk.Filename );
+				UE_LOG( LogProfilerClient, Log, TEXT( "Received a bad chunk of file, resending: %5i, %6u, %10u, %s" ), FileChunk.Data.Num(), ReceivedFileInfo->Progress, FileChunkHeader.FileSize, *FileChunk.Filename );
 			}
 		}
 	}
@@ -525,9 +537,11 @@ void FProfilerClientManager::HandleServicePingMessage( const FProfilerServicePin
 		TArray<FMessageAddress> Instances;
 		for (auto It = Connections.CreateConstIterator(); It; ++It)
 		{
-			Instances.Add(It.Value().ServiceAddress);
+			Instances.Add(It.Value().ProfilerServiceAddress);
 		}
 		MessageEndpoint->Send(new FProfilerServicePong(), Instances);
+
+		UE_LOG( LogProfilerClient, Verbose, TEXT( "Ping GetSender: %s" ), *Context->GetSender().ToString() );
 	}
 #endif
 }
@@ -602,25 +616,35 @@ bool FProfilerClientManager::HandleMessagesTicker( float DeltaTime )
 		Connection.PendingMessages.GenerateKeyArray( Frames );
 		Frames.Sort();
 
+		// MessageBus sends all data in out of order fashion.
+		// We buffer frame to make sure that all frames are received in the proper order.
+		const int32 NUM_BUFFERED_FRAMES = 15;
+
 		for( int32 Index = 0; Index < Frames.Num(); Index++ )
 		{
+			if (Connection.PendingMessages.Num() < NUM_BUFFERED_FRAMES)
+			{
+				break;
+			}
+
 			int64 FrameNum = Frames[Index];
 			TArray<uint8>& Data = *Connection.PendingMessages.Find( FrameNum );
 
-			// pass the data to the visualization code
-			FArrayReader Reader(true);
-			Reader.Append(Data);
+			// Pass the data to the visualization code.
+			FMemoryReader MemoryReader( Data );
 
-			int64 Size = Reader.TotalSize();
+			int64 Size = MemoryReader.TotalSize();
 			FStatsReadStream& Stream = Connection.Stream;
 
+			UE_LOG( LogProfilerClient, VeryVerbose, TEXT( "Frame=%i/%i, FNamesIndexMap=%i, CurrentMetadataSize=%i" ), FrameNum, Frames.Num(), Connection.Stream.FNamesIndexMap.Num(), Connection.CurrentThreadState.ShortNameToLongName.Num() );
+
 			// read in the data and post if we reach a 
 			{
 				SCOPE_CYCLE_COUNTER(STAT_PC_ReadStatMessages);
-				while(Reader.Tell() < Size)
+				while(MemoryReader.Tell() < Size)
 				{
 					// read the message
-					FStatMessage Message(Stream.ReadMessage(Reader));
+					FStatMessage Message(Stream.ReadMessage(MemoryReader));
 
 					if (Message.NameAndInfo.GetField<EStatOperation>() == EStatOperation::AdvanceFrameEventGameThread)
 					{
@@ -641,7 +665,7 @@ bool FProfilerClientManager::HandleMessagesTicker( float DeltaTime )
 			}
 
 			// send the data out
-			ProfilerDataDelegate.Broadcast(Connection.InstanceId, Connection.CurrentData,0.0f);
+			ProfilerDataDelegate.Broadcast( Connection.InstanceId, Connection.CurrentData, 0.0f );
 
 			Connection.PendingMessages.Remove( FrameNum );
 		}
@@ -655,7 +679,7 @@ bool FProfilerClientManager::HandleMessagesTicker( float DeltaTime )
 
 		if( ReceivedFileInfo.IsTimedOut() )
 		{
-			UE_LOG(LogProfile, Log, TEXT( "File service-client timed out, aborted: %s" ), *Filename );
+			UE_LOG( LogProfilerClient, Log, TEXT( "File service-client timed out, aborted: %s" ), *Filename );
 			FailedTransfer.Add( Filename );
 
 			delete ReceivedFileInfo.FileWriter;
@@ -677,12 +701,13 @@ void FProfilerClientManager::HandleServicePreviewAckMessage( const FProfilerServ
 	if (ActiveSessionId.IsValid() && Connections.Find(Message.InstanceId) != NULL)
 	{
 		FServiceConnection& Connection = *Connections.Find(Message.InstanceId);
-		Connection.CurrentFrame = Message.Frame;//
+
+		UE_LOG( LogProfilerClient, Verbose, TEXT( "PreviewAck InstanceId: %s, GetSender: %s" ), *Message.InstanceId.ToString(), *Context->GetSender().ToString() );
 	}
 #endif
 }
 
-void FProfilerClientManager::HandleServiceData2Message( const FProfilerServiceData2& Message, const IMessageContextRef& Context )
+void FProfilerClientManager::HandleProfilerServiceData2Message( const FProfilerServiceData2& Message, const IMessageContextRef& Context )
 {
 #if STATS
 	SCOPE_CYCLE_COUNTER(STAT_PC_HandleDataReceived);
@@ -690,10 +715,10 @@ void FProfilerClientManager::HandleServiceData2Message( const FProfilerServiceDa
 	{
 		FServiceConnection& Connection = *Connections.Find(Message.InstanceId);
 
-		// add the message to the connections queue
-		TArray<uint8> Data;
-		Data.Append(Message.Data);
-		Connection.PendingMessages.Add(Message.Frame, Data);
+		// Add the message to the connections queue.
+		Connection.PendingMessages.Add( Message.Frame, Message.Data );
+
+		UE_LOG( LogProfilerClient, VeryVerbose, TEXT( "Frame:%i, Data.Num:%i, InstanceId:%s" ), Message.Frame, Message.Data.Num(), *Message.InstanceId.ToString() );
 	}
 #endif
 }
@@ -1016,7 +1041,7 @@ bool FServiceConnection::ReadAndConvertStatMessages( FArchive& Reader, bool bUse
 				return true;
 			}
 
-			FMemoryReader MemoryReader( DestData, true );
+			FMemoryReader MemoryReader( DestData );
 
 			while( MemoryReader.Tell() < MemoryReader.TotalSize() )
 			{
diff --git a/Engine/Source/Developer/ProfilerClient/Private/ProfilerClientManager.h b/Engine/Source/Developer/ProfilerClient/Private/ProfilerClientManager.h
--- a/Engine/Source/Developer/ProfilerClient/Private/ProfilerClientManager.h
+++ b/Engine/Source/Developer/ProfilerClient/Private/ProfilerClientManager.h
@@ -2,7 +2,7 @@
 
 #pragma once
 
-
+// #YRX_Profiler: 2015-05-13 Remove it later.
 #define PROFILER_THREADED_LOAD 1
 
 
@@ -14,7 +14,6 @@ struct FServiceConnection
 		, ReadMessages(0)
 	{
 		MetaData.CriticalSection = &CriticalSection;
-		CurrentFrame = 0;
 	}
 
 	~FServiceConnection()
@@ -25,14 +24,13 @@ struct FServiceConnection
 	FServiceConnection(const FServiceConnection& InConnnection)
 	{
 		MetaData.CriticalSection = &CriticalSection;
-		CurrentFrame = 0;
 	}
 
 	/** Instance Id */
 	FGuid InstanceId;
 
 	/** Service endpoint */
-	FMessageAddress ServiceAddress;
+	FMessageAddress ProfilerServiceAddress;
 
 	/** Descriptions for the stats */
 	FStatMetaData MetaData;
@@ -41,7 +39,7 @@ struct FServiceConnection
 	FProfilerDataFrame CurrentData;
 
 	/** Initialize the connection from the given authorization message and context */
-	void Initialize( const FProfilerServiceAuthorize2& Message, const IMessageContextRef& Context );
+	void Initialize( const FProfilerServiceAuthorize& Message, const IMessageContextRef& Context );
 
 #if STATS
 	/** Current stats data */
@@ -68,9 +66,6 @@ struct FServiceConnection
 	/** Current data loading progress. */
 	float DataLoadingProgress;
 
-	/** Current frame of data to process */
-	int64 CurrentFrame;
-
 	/** Current number of read messages. */
 	uint64 ReadMessages;
 
@@ -221,7 +216,7 @@ private:
 	void HandleMessageBusShutdown();
 
 	// Handles FProfilerServiceAuthorize2 messages.
-	void HandleServiceAuthorize2Message( const FProfilerServiceAuthorize2& Message, const IMessageContextRef& Context );
+	void HandleServiceAuthorizeMessage( const FProfilerServiceAuthorize& Message, const IMessageContextRef& Context );
 
 	// Handles FProfilerServiceFileChunk messages.
 	void HandleServiceFileChunk( const FProfilerServiceFileChunk& FileChunk, const IMessageContextRef& Context );
@@ -233,7 +228,7 @@ private:
 	bool HandleTicker( float DeltaTime );
 
 	// Handles FProfilerServiceData2 messages.
-	void HandleServiceData2Message( const FProfilerServiceData2& Message, const IMessageContextRef& Context );
+	void HandleProfilerServiceData2Message( const FProfilerServiceData2& Message, const IMessageContextRef& Context );
 
 	// Handles FProfilerServicePreviewAck messages
 	void HandleServicePreviewAckMessage( const FProfilerServicePreviewAck& Messsage, const IMessageContextRef& Context );
@@ -342,6 +337,9 @@ private:
 	/** Handle to the registered MessageDelegate. */
 	FDelegateHandle MessageDelegateHandle;
 
+	/** Handle to the registered OnShutdown for Message Bus. */
+	FDelegateHandle OnShutdownMessageBusDelegateHandle;
+
 #if PROFILER_THREADED_LOAD
 	/** Loads a file asynchronously */
 	bool AsyncLoad();
diff --git a/Engine/Source/Developer/ProfilerMessages/Classes/ProfilerServiceMessages.h b/Engine/Source/Developer/ProfilerMessages/Classes/ProfilerServiceMessages.h
--- a/Engine/Source/Developer/ProfilerMessages/Classes/ProfilerServiceMessages.h
+++ b/Engine/Source/Developer/ProfilerMessages/Classes/ProfilerServiceMessages.h
@@ -22,40 +22,6 @@ struct FProfilerServiceAuthorize
 	UPROPERTY()
 	FGuid InstanceId;
 
-	/**
-	 * Default constructor.
-	 */
-	FProfilerServiceAuthorize( ) { }
-
-	/**
-	 * Creates and initializes a new instance.
-	 */
-	FProfilerServiceAuthorize( const FGuid& InSessionId, const FGuid& InInstanceId )
-		: SessionId(InSessionId)
-		, InstanceId(InInstanceId)
-	{
-	}
-
-};
-
-
-/**
- */
-USTRUCT()
-struct FProfilerServiceAuthorize2
-{
-	GENERATED_USTRUCT_BODY()
-
-	/**
-	 */
-	UPROPERTY()
-	FGuid SessionId;
-
-	/**
-	 */
-	UPROPERTY()
-	FGuid InstanceId;
-
 	/**
 	 */
 	UPROPERTY()
@@ -64,12 +30,12 @@ struct FProfilerServiceAuthorize2
 	/**
 	 * Default constructor.
 	 */
-	FProfilerServiceAuthorize2( ) { }
+	FProfilerServiceAuthorize( ) { }
 
 	/**
 	 * Creates and initializes a new instance.
 	 */
-	FProfilerServiceAuthorize2( const FGuid& InSessionId, const FGuid& InInstanceId, const TArray<uint8>& InData )
+	FProfilerServiceAuthorize( const FGuid& InSessionId, const FGuid& InInstanceId, const TArray<uint8>& InData )
 		: SessionId(InSessionId)
 		, InstanceId(InInstanceId)
 	{
@@ -163,10 +129,6 @@ struct FProfilerServicePreviewAck
 	UPROPERTY()
 	FGuid InstanceId;
 
-	UPROPERTY()
-	int64 Frame;
-
-
 	/**
 	 * Default constructor.
 	 */
@@ -175,9 +137,8 @@ struct FProfilerServicePreviewAck
 	/**
 	 * Creates and initializes a new instance.
 	 */
-	FProfilerServicePreviewAck( const FGuid& InInstance, int64 InFrame )
+	FProfilerServicePreviewAck( const FGuid& InInstance )
 		: InstanceId(InInstance)
-		, Frame(InFrame)
 	{
 	}
 };
diff --git a/Engine/Source/Developer/ProfilerService/Private/ProfilerServiceManager.cpp b/Engine/Source/Developer/ProfilerService/Private/ProfilerServiceManager.cpp
--- a/Engine/Source/Developer/ProfilerService/Private/ProfilerServiceManager.cpp
+++ b/Engine/Source/Developer/ProfilerService/Private/ProfilerServiceManager.cpp
@@ -7,8 +7,7 @@
 #include "SecureHash.h"
 
 
-DEFINE_LOG_CATEGORY_STATIC(LogProfile, Log, All);
-
+DEFINE_LOG_CATEGORY_STATIC( LogProfilerService, Log, All );
 
 /** 
  * Thread used to read, prepare and send files through the message bus.
@@ -51,7 +50,7 @@ public:
 			FReaderAndAddress& ReaderAndAddress = It.Value();
 			DeleteFileReader( ReaderAndAddress );
 
-			UE_LOG(LogProfile, Log, TEXT( "File service-client sending aborted (srv): %s" ), *It.Key() );
+			UE_LOG( LogProfilerService, Log, TEXT( "File service-client sending aborted (srv): %s" ), *It.Key() );
 		}
 
 		FPlatformProcess::ReturnSynchEventToPool( WorkEvent );
@@ -131,19 +130,19 @@ public:
 		const FString PathName = FPaths::ProfilingDir() + TEXT("UnrealStats/");
 		FString StatFilepath = PathName + StatFilename;
 
-		UE_LOG(LogProfile, Log, TEXT( "Opening stats file for service-client sending: %s" ), *StatFilepath );
+		UE_LOG( LogProfilerService, Log, TEXT( "Opening stats file for service-client sending: %s" ), *StatFilepath );
 
 		const int64 FileSize = IFileManager::Get().FileSize(*StatFilepath);
 		if( FileSize < 4 )
 		{
-			UE_LOG(LogProfile, Error, TEXT( "Could not open: %s" ), *StatFilepath );
+			UE_LOG( LogProfilerService, Error, TEXT( "Could not open: %s" ), *StatFilepath );
 			return;
 		}
 
 		FArchive* FileReader = IFileManager::Get().CreateFileReader(*StatFilepath);
 		if( !FileReader )
 		{
-			UE_LOG(LogProfile, Error, TEXT( "Could not open: %s" ), *StatFilepath );
+			UE_LOG( LogProfilerService, Error, TEXT( "Could not open: %s" ), *StatFilepath );
 			return;
 		}
 
@@ -212,7 +211,7 @@ public:
 		FReaderAndAddress ReaderAndAddress = ActiveTransfers.FindAndRemoveChecked( Filename );
 		DeleteFileReader( ReaderAndAddress );
 
-		UE_LOG(LogProfile, Log, TEXT( "File service-client sent successfully : %s" ), *Filename );
+		UE_LOG( LogProfilerService, Log, TEXT( "File service-client sent successfully : %s" ), *Filename );
 	}
 
 	/** Aborts file sending to the specified client, probably client disconnected or exited. */
@@ -225,7 +224,7 @@ public:
 			FReaderAndAddress& ReaderAndAddress = It.Value();
 			if( ReaderAndAddress.Value == Recipient )
 			{
-				UE_LOG(LogProfile, Log, TEXT( "File service-client sending aborted (cl): %s" ), *It.Key() );
+				UE_LOG( LogProfilerService, Log, TEXT( "File service-client sending aborted (cl): %s" ), *It.Key() );
 				FReaderAndAddress ActiveReaderAndAddress = ActiveTransfers.FindAndRemoveChecked( It.Key() );
 				DeleteFileReader( ActiveReaderAndAddress );
 			}
@@ -306,14 +305,34 @@ protected:
 	TMap<FString,FReaderAndAddress> ActiveTransfers;
 };
 
+void FProfilerServiceManager::HandleServiceFileChunkMessage( const FProfilerServiceFileChunk& Message, const IMessageContextRef& Context )
+{
+	FMemoryReader Reader( Message.Header );
+	FProfilerFileChunkHeader Header;
+	Reader << Header;
+	Header.Validate();
+
+	if (Header.ChunkType == EProfilerFileChunkType::SendChunk)
+	{
+		// Send this file chunk again.
+		FileTransferRunnable->EnqueueFileChunkToSend( new FProfilerServiceFileChunk( Message, FProfilerServiceFileChunk::FNullTag() ), true );
+	}
+	else if (Header.ChunkType == EProfilerFileChunkType::FinalizeFile)
+	{
+		// Finalize file.
+		FileTransferRunnable->FinalizeFileSending( Message.Filename );
+	}
+}
+
+
 /* FProfilerServiceManager structors
  *****************************************************************************/
 
 FProfilerServiceManager::FProfilerServiceManager()
+	: FileTransferRunnable( nullptr )
+	, MetadataSize( 0 )
 {
-	MetaData.SecondsPerCycle = FPlatformTime::GetSecondsPerCycle();
 	PingDelegate = FTickerDelegate::CreateRaw(this, &FProfilerServiceManager::HandlePing);
-	DataFrame.Frame = 0;
 }
 
 /* IProfilerServiceManager interface
@@ -392,6 +411,7 @@ void FProfilerServiceManager::AddNewFrameHandleStatsThread()
 	const FStatsThreadState& Stats = FStatsThreadState::GetLocalState();
 	NewFrameDelegateHandle = Stats.NewFrameDelegate.AddRaw( this, &FProfilerServiceManager::HandleNewFrame );
 	StatsMasterEnableAdd();
+	MetadataSize = 0;
 #endif // STATS
 }
 
@@ -401,6 +421,7 @@ void FProfilerServiceManager::RemoveNewFrameHandleStatsThread()
 	const FStatsThreadState& Stats = FStatsThreadState::GetLocalState();
 	Stats.NewFrameDelegate.Remove( NewFrameDelegateHandle );
 	StatsMasterEnableSubtract();
+	MetadataSize = 0;
 #endif // STATS
 }
 
@@ -408,7 +429,7 @@ void FProfilerServiceManager::SetPreviewState( const FMessageAddress& ClientAddr
 {
 #if STATS
 	FClientData* Client = ClientData.Find( ClientAddress );
-	if( Client )
+	if (MessageEndpoint.IsValid() && Client)
 	{
 		const bool bIsPreviewing = Client->Preview;
 
@@ -416,10 +437,10 @@ void FProfilerServiceManager::SetPreviewState( const FMessageAddress& ClientAddr
 		{
 			if( bRequestedPreviewState )
 			{
-				// enable stat capture
+				// Enable stat capture.
 				if (PreviewClients.Num() == 0)
 				{
-					FGraphEventRef CompletionEvent = FSimpleDelegateGraphTask::CreateAndDispatchWhenReady
+					FSimpleDelegateGraphTask::CreateAndDispatchWhenReady
 					(
 						FSimpleDelegateGraphTask::FDelegate::CreateRaw( this, &FProfilerServiceManager::AddNewFrameHandleStatsThread ),
 						TStatId(), nullptr,
@@ -428,31 +449,29 @@ void FProfilerServiceManager::SetPreviewState( const FMessageAddress& ClientAddr
 				}
 				PreviewClients.Add(ClientAddress);
 				Client->Preview = true;
-				if (MessageEndpoint.IsValid())
-				{
-					Client->CurrentFrame = FStats::GameThreadStatsFrame;
-					MessageEndpoint->Send( new FProfilerServicePreviewAck( InstanceId, FStats::GameThreadStatsFrame ), ClientAddress );
-				}
+
+				MessageEndpoint->Send( new FProfilerServicePreviewAck( InstanceId ), ClientAddress );
 			}
 			else
 			{
 				PreviewClients.Remove(ClientAddress);
 				Client->Preview = false;
 
-				// stop the ping messages if we have no clients
+				// Disable stat capture.
 				if (PreviewClients.Num() == 0)
 				{
-					// disable stat capture
-					FGraphEventRef CompletionEvent = FSimpleDelegateGraphTask::CreateAndDispatchWhenReady
+					FSimpleDelegateGraphTask::CreateAndDispatchWhenReady
 					(
 						FSimpleDelegateGraphTask::FDelegate::CreateRaw( this, &FProfilerServiceManager::RemoveNewFrameHandleStatsThread ),
 						TStatId(), nullptr,
 						FPlatformProcess::SupportsMultithreading() ? ENamedThreads::StatsThread : ENamedThreads::GameThread
 					);
 					
-				}
+				}	
 			}
 		}
+
+		UE_LOG( LogProfilerService, Verbose, TEXT( "SetPreviewState: %i, InstanceId: %s, ClientAddress: %s" ), (int32)bRequestedPreviewState, *InstanceId.ToString(), *ClientAddress.ToString() );
 	}
 #endif
 }
@@ -473,13 +492,13 @@ bool FProfilerServiceManager::HandlePing( float DeltaTime )
 		{
 			Iter.Value().Active = false;
 			Clients.Add(Iter.Key());
+			UE_LOG( LogProfilerService, Verbose, TEXT( "Ping Active 0: %s, InstanceId: %s, ClientAddress: %s" ), *Iter.Key().ToString(), *InstanceId.ToString(), *ClientAddress.ToString() );
 		}
 		else
 		{
-			if (PreviewClients.Contains(ClientAddress))
-			{
-				PreviewClients.Remove(ClientAddress);	
-			}
+			UE_LOG( LogProfilerService, Verbose, TEXT( "Ping Remove: %s, InstanceId: %s, ClientAddress: %s" ), *Iter.Key().ToString(), *InstanceId.ToString(), *ClientAddress.ToString() );
+			SetPreviewState( ClientAddress, false );
+
 			Iter.RemoveCurrent();
 			FileTransferRunnable->AbortFileSending( ClientAddress );
 		}
@@ -490,8 +509,9 @@ bool FProfilerServiceManager::HandlePing( float DeltaTime )
 	{
 		MessageEndpoint->Send(new FProfilerServicePing(), Clients);
 	}
-#endif
 	return (ClientData.Num() > 0);
+#endif
+	return false;
 }
 
 
@@ -505,11 +525,12 @@ void FProfilerServiceManager::HandleServiceCaptureMessage( const FProfilerServic
 	{
 		if( bRequestedCaptureState && !bIsCapturing )
 		{
-			UE_LOG(LogProfile, Log, TEXT("StartCapture") );
+			UE_LOG( LogProfilerService, Verbose, TEXT( "StartCapture, InstanceId: %s, GetSender: %s" ), *InstanceId.ToString(), *Context->GetSender().ToString() );
 			StartCapture();
 		}
 		else if( !bRequestedCaptureState && bIsCapturing )
 		{
+			UE_LOG( LogProfilerService, Verbose, TEXT( "StopCapture, InstanceId: %s, GetSender: %s" ), *InstanceId.ToString(), *Context->GetSender().ToString() );
 			StopCapture();
 		}
 	}
@@ -524,6 +545,7 @@ void FProfilerServiceManager::HandleServicePongMessage( const FProfilerServicePo
 	if (Data != nullptr)
 	{
 		Data->Active = true;
+		UE_LOG( LogProfilerService, Verbose, TEXT( "Pong InstanceId: %s, GetSender: %s" ), *InstanceId.ToString(), *Context->GetSender().ToString() );
 	}
 }
 
@@ -546,52 +568,32 @@ void FProfilerServiceManager::HandleServiceRequestMessage( const FProfilerServic
 	}
 }
 
-void FProfilerServiceManager::HandleServiceFileChunkMessage( const FProfilerServiceFileChunk& Message, const IMessageContextRef& Context )
-{
-	FMemoryReader Reader(Message.Header);
-	FProfilerFileChunkHeader Header;
-	Reader << Header;
-	Header.Validate();
 
-	if( Header.ChunkType == EProfilerFileChunkType::SendChunk )
-	{
-		// Send this file chunk again.
-		FileTransferRunnable->EnqueueFileChunkToSend( new FProfilerServiceFileChunk(Message,FProfilerServiceFileChunk::FNullTag()), true );
-	}
-	else if( Header.ChunkType == EProfilerFileChunkType::FinalizeFile )
-	{
-		// Finalize file.
-		FileTransferRunnable->FinalizeFileSending( Message.Filename );
-	}
-}
 
 void FProfilerServiceManager::HandleServiceSubscribeMessage( const FProfilerServiceSubscribe& Message, const IMessageContextRef& Context )
 {
 #if STATS
-	const FMessageAddress& MsgAddress = Context->GetSender();
-	if( Message.SessionId == SessionId && Message.InstanceId == InstanceId && !ClientData.Contains( MsgAddress ) )
+	const FMessageAddress& SenderAddress = Context->GetSender();
+	if( MessageEndpoint.IsValid() && Message.SessionId == SessionId && Message.InstanceId == InstanceId && !ClientData.Contains( SenderAddress ) )
 	{
-		UE_LOG(LogProfile, Log, TEXT("Added a client" ));
+		UE_LOG( LogProfilerService, Log, TEXT( "Subscribe Session: %s, Instance: %s" ), *SessionId.ToString(), *InstanceId.ToString() );
 
 		FClientData Data;
 		Data.Active = true;
 		Data.Preview = false;
 		Data.StatsWriteFile.WriteHeader();
 
-		// add to the client list
-		ClientData.Add( MsgAddress, Data );
-		// send authorized and stat descriptions
-		const TArray<uint8>& OutData = ClientData.Find( MsgAddress )->StatsWriteFile.GetOutData();
-		MessageEndpoint->Send( new FProfilerServiceAuthorize2( SessionId, InstanceId, OutData ), MsgAddress );
+		// Add to the client list.
+		ClientData.Add( SenderAddress, Data );
 
-		// Not thread-safe, but reading the number of elements should be ok.
-		const FStatsThreadState& Stats = FStatsThreadState::GetLocalState();
-		ClientData.Find( MsgAddress )->MetadataSize = Stats.ShortNameToLongName.Num();
+		// Send authorized and stat descriptions.
+		const TArray<uint8>& OutData = ClientData.Find( SenderAddress )->StatsWriteFile.GetOutData();
+		MessageEndpoint->Send( new FProfilerServiceAuthorize( SessionId, InstanceId, OutData ), SenderAddress );
 
-		// initiate the ping callback
+		// Initiate the ping callback
 		if (ClientData.Num() == 1)
 		{
-			PingDelegateHandle = FTicker::GetCoreTicker().AddTicker(PingDelegate, 15.0f);
+			PingDelegateHandle = FTicker::GetCoreTicker().AddTicker(PingDelegate, 5.0f);
 		}
 	}
 #endif
@@ -600,10 +602,11 @@ void FProfilerServiceManager::HandleServiceSubscribeMessage( const FProfilerServ
 
 void FProfilerServiceManager::HandleServiceUnsubscribeMessage( const FProfilerServiceUnsubscribe& Message, const IMessageContextRef& Context )
 {
+#if	STATS
 	const FMessageAddress SenderAddress = Context->GetSender();
 	if (Message.SessionId == SessionId && Message.InstanceId == InstanceId)
 	{
-		UE_LOG(LogProfile, Log, TEXT("Removed a client"));
+		UE_LOG( LogProfilerService, Log, TEXT( "Unsubscribe Session: %s, Instance: %s" ), *SessionId.ToString(), *InstanceId.ToString() );
 
 		// clear out any previews
 		while (PreviewClients.Num() > 0)
@@ -621,35 +624,47 @@ void FProfilerServiceManager::HandleServiceUnsubscribeMessage( const FProfilerSe
 			FTicker::GetCoreTicker().RemoveTicker(PingDelegateHandle);
 		}
 	}
+#endif // STATS
 }
 
 void FProfilerServiceManager::HandleNewFrame(int64 Frame)
 {
 	// Called from the stats thread.
 #if STATS
-	// package it up and send to the clients
-	if (MessageEndpoint.IsValid())
-	{
-		const FStatsThreadState& Stats = FStatsThreadState::GetLocalState();
-		const int32 CurrentMetadataSize = Stats.ShortNameToLongName.Num();
+	const FStatsThreadState& Stats = FStatsThreadState::GetLocalState();
+	const int32 CurrentMetadataSize = Stats.ShortNameToLongName.Num();
 
-		// update preview clients with the current data
-		for (auto It = PreviewClients.CreateConstIterator(); It; ++It)
-		{
-			FClientData& Client = *ClientData.Find(*It);
+	bool bNeedFullMetadata = false;
+	if (MetadataSize < CurrentMetadataSize)
+	{
+		// Write the whole metadata.
+		bNeedFullMetadata = true;
+		MetadataSize = CurrentMetadataSize;
+	}
 
-			Client.StatsWriteFile.ResetData();
-			bool bNeedFullMetadata = false;
-			if( Client.MetadataSize < CurrentMetadataSize )
-			{
-				// Write the whole metadata.
-				bNeedFullMetadata = true;
-				Client.MetadataSize = CurrentMetadataSize;
-			}
+	// Create a temporary client data and prepare all data.
+	FClientData* ToGameThread = new FClientData;
+	ToGameThread->StatsWriteFile.ResetData();
+	ToGameThread->StatsWriteFile.WriteFrame( Frame, bNeedFullMetadata );
+	ToGameThread->Frame = Frame;
+
+	// Send to the game thread. PreviewClients is not thread-safe, so we cannot send the data here.
+	FSimpleDelegateGraphTask::CreateAndDispatchWhenReady
+	(
+		FSimpleDelegateGraphTask::FDelegate::CreateRaw( this, &FProfilerServiceManager::HandleNewFrameGT, ToGameThread ),
+		TStatId(), nullptr, ENamedThreads::GameThread
+	);
+#endif
+}
 
-			Client.StatsWriteFile.WriteFrame( Frame, bNeedFullMetadata );
-			MessageEndpoint->Send( new FProfilerServiceData2( InstanceId, Frame, Client.StatsWriteFile.GetOutData() ), PreviewClients );
-		}
+void FProfilerServiceManager::HandleNewFrameGT( FClientData* ToGameThread )
+{
+	if (MessageEndpoint.IsValid())
+	{
+		// Update preview clients with the current data.
+		UE_LOG( LogProfilerService, VeryVerbose, TEXT( "Frame: %i, MetadataSize: %i, PreviewClients: %i" ), ToGameThread->Frame, MetadataSize, PreviewClients.Num() );
+		MessageEndpoint->Send( new FProfilerServiceData2( InstanceId, ToGameThread->Frame, ToGameThread->StatsWriteFile.GetOutData() ), PreviewClients );
 	}
-#endif
-}
\ No newline at end of file
+
+	delete ToGameThread;
+}
diff --git a/Engine/Source/Developer/ProfilerService/Private/ProfilerServiceManager.h b/Engine/Source/Developer/ProfilerService/Private/ProfilerServiceManager.h
--- a/Engine/Source/Developer/ProfilerService/Private/ProfilerServiceManager.h
+++ b/Engine/Source/Developer/ProfilerService/Private/ProfilerServiceManager.h
@@ -3,27 +3,35 @@
 #pragma once
 
 
+#if STATS
+
 /**
  * struct that holds the client information
  */
 struct FClientData
 {
-	/** Connection is active */
-	bool Active;
-	/** Connection is previewing */
-	bool Preview;
-
-#if STATS
+	/** Stats data. */
 	FStatsWriteFile StatsWriteFile;
 
-	/** Current frame of data sent */
-	int64 CurrentFrame;
+	/** Frame. */
+	int64 Frame;
 
-	/** Stats metadata size. */
-	int32 MetadataSize;
-#endif
+	/** Connection is active. */
+	bool Active;
+
+	/** Connection is previewing. */
+	bool Preview;
+
+	/** Default constructor. */
+	FClientData()
+		: Frame(-1)
+		, Active( false )
+		, Preview( false )
+	{}
 };
 
+#endif //STATS
+
 
 /**
  * Implements the Profile Service Manager
@@ -42,20 +50,8 @@ public:
 public:
 
 	// Begin IProfilerServiceManager interface
-
 	virtual void StartCapture() override;
-	
 	virtual void StopCapture() override;
-
-	virtual FStatMetaData& GetStatMetaData() override
-	{
-		return MetaData;
-	}
-
-	virtual FProfilerDataDelegate& OnProfilerData() override
-	{
-		return ProfilerDataDelegate;
-	}
 	// End IProfilerServiceManager interface
 
 	/**
@@ -83,7 +79,6 @@ private:
 	/** Callback for a tick, used to ping the clients */
 	bool HandlePing( float DeltaTime );
 
-private:
 
 	// Handles FProfilerServiceCapture messages.
 	void HandleServiceCaptureMessage( const FProfilerServiceCapture& Message, const IMessageContextRef& Context );
@@ -108,13 +103,14 @@ private:
 
 	/** Handles a new frame from the stats system. Called from the stats thread. */
 	void HandleNewFrame(int64 Frame);
+
+	/** Handles a new frame from the stats system. Called from the game thread. */
+	void HandleNewFrameGT( FClientData* ToGameThread );
 	
 	void AddNewFrameHandleStatsThread();
 
 	void RemoveNewFrameHandleStatsThread();
 
-private:
-
 	/** Holds the messaging endpoint. */
 	FMessageEndpointPtr MessageEndpoint;
 
@@ -125,8 +121,10 @@ private:
 	/** Holds the message addresses for registered clients */
 	TArray<FMessageAddress> PreviewClients;
 
+#if	STATS
 	/** Holds the client data for registered clients */
 	TMap<FMessageAddress, FClientData> ClientData;
+#endif // STATS
 
 	/** Thread used to read, prepare and send file chunks through the message bus. */
 	class FFileTransferRunnable* FileTransferRunnable;
@@ -134,14 +132,8 @@ private:
 	/** Filename of last capture file. */
 	FString LastStatsFilename;
 
-	/** Stat meta data */
-	FStatMetaData MetaData;
-
-	/** Delegate for notifying clients of received data */
-	FProfilerDataDelegate ProfilerDataDelegate;
-
-	/** Frame of data */
-	FProfilerDataFrame DataFrame;
+	/** Size of the stats metadata. */
+	int32 MetadataSize;
 
 	/** Holds a delegate to be invoked for client pings */
 	FTickerDelegate PingDelegate;
diff --git a/Engine/Source/Developer/ProfilerService/Public/Interfaces/IProfilerServiceManager.h b/Engine/Source/Developer/ProfilerService/Public/Interfaces/IProfilerServiceManager.h
--- a/Engine/Source/Developer/ProfilerService/Public/Interfaces/IProfilerServiceManager.h
+++ b/Engine/Source/Developer/ProfilerService/Public/Interfaces/IProfilerServiceManager.h
@@ -223,142 +223,6 @@ FORCEINLINE FArchive& operator<<(FArchive& Ar, FStatMetaData& Data)
 	return Ar;
 }
 
-
-/**
- * Profiler file type
- */
-namespace EProfilerFileType
-{
-	enum Type
-	{
-		PFT_Binary,
-		PFT_CSV,
-	};
-}
-
-/**
- * Profiler file header
- */
-struct FProfilerDataHeader
-{
-	/** Magice value to determine it is a profiler file */
-	uint32 Magic;
-
-	/** Current version of the file */
-	uint32 Version;
-
-	/** Session id the file was generated from */
-	FGuid SessionId;
-
-	/** Instance id the file was generated from */
-	FGuid InstanceId;
-
-	/** Total amount of capture data.  Note: this is different from the file size */
-	int64 TotalCaptureData;
-};
-
-FORCEINLINE FArchive& operator<<(FArchive& Ar, FProfilerDataHeader& Data)
-{
-	Ar << Data.Magic;
-	Ar << Data.Version;
-	Ar << Data.SessionId;
-	Ar << Data.InstanceId;
-	Ar << Data.TotalCaptureData;
-
-	return Ar;
-}
-
-#if STATS
-FORCEINLINE FArchive& operator<<(FArchive& Ar, FStatNameAndInfo& NameAndInfo)
-{
-	if (Ar.IsLoading())
-	{
-		int32 Index = 0;
-		Ar << Index;
-		int32 Number = 0;
-		Ar << Number;
-		FName TheFName;
-		FString Name;
-		Ar << Name;
-		TheFName = FName(*Name);
-		Number &= ~ (EStatMetaFlags::SendingFName << (EStatMetaFlags::Shift + EStatAllFields::StartShift));
-		NameAndInfo = FStatNameAndInfo(TheFName, false);
-		NameAndInfo.SetNumberDirect(Number);
-	}
-	else
-	{
-		FName RawName = NameAndInfo.GetRawName();
-		int32 Index = RawName.GetComparisonIndex();
-		Ar << Index;
-		int32 Number = NameAndInfo.GetRawNumber();
-		Number |= EStatMetaFlags::SendingFName << (EStatMetaFlags::Shift + EStatAllFields::StartShift);
-		Ar << Number;
-		FString Name = RawName.ToString();
-		Ar << Name;
-	}
-
-	return Ar;
-}
-
-FORCEINLINE FArchive& operator<<(FArchive& Ar, FStatMessage& Data)
-{
-	Ar << Data.NameAndInfo;
-
-	if (Ar.IsLoading())
-	{
-		Data.Clear();
-		switch (Data.NameAndInfo.GetField<EStatDataType>())
-		{
-		case EStatDataType::ST_int64:
-			{
-				int64 Payload = 0;
-				Ar << Payload;
-				Data.GetValue_int64() = Payload;
-			}
-			break;
-		case EStatDataType::ST_double:
-			{
-				double Payload = 0;
-				Ar << Payload;
-				Data.GetValue_double() = Payload;
-			}
-			break;
-		case EStatDataType::ST_FName:
-			FStatNameAndInfo Payload;
-			Ar << Payload;
-			Data.GetValue_FName() = Payload.GetRawName();
-			break;
-		}
-
-	}
-	else
-	{
-		switch (Data.NameAndInfo.GetField<EStatDataType>())
-		{
-		case EStatDataType::ST_int64:
-			{
-				int64 Payload = Data.GetValue_int64();
-				Ar << Payload;
-			}
-			break;
-		case EStatDataType::ST_double:
-			{
-				double Payload = Data.GetValue_double();
-				Ar << Payload;
-			}
-			break;
-		case EStatDataType::ST_FName:
-			FStatNameAndInfo Payload(Data.GetValue_FName(), false);
-			Ar << Payload;
-			break;
-		}
-	}
-
-	return Ar;
-}
-#endif
-
-
 /**
  * Profiler service request type
  */
@@ -476,22 +340,6 @@ public:
 
 	/** Stops a file capture. */
 	virtual void StopCapture() = 0;
-
-	/**
-	 * Gets the description for the given stat id.
-	 *
-	 * @param StatId IDd of the statistic description to retrieve.
-	 * @return the FStatMetaData struct with the description.
-	 */
-	virtual FStatMetaData& GetStatMetaData() = 0;
-
-	/**
-	 * Retrieves the profiler data delegate
-	 *
-	 * @return profiler data delegate
-	 */
-	virtual FProfilerDataDelegate& OnProfilerData() = 0;
-
 public:
 
 	/** Virtual destructor. */
