Profiler - Profiler will now load all frames, previously the last frame was ignored

diff --git a/Engine/Source/Developer/ProfilerClient/Private/ProfilerClientManager.cpp b/Engine/Source/Developer/ProfilerClient/Private/ProfilerClientManager.cpp
--- a/Engine/Source/Developer/ProfilerClient/Private/ProfilerClientManager.cpp
+++ b/Engine/Source/Developer/ProfilerClient/Private/ProfilerClientManager.cpp
@@ -1085,6 +1085,8 @@ bool FServiceConnection::ReadAndConvertStatMessages( FArchive& Reader, bool bUse
 			Reader << UncompressedData;
 			if( UncompressedData.HasReachedEndOfCompressedData() )
 			{
+				static FStatNameAndInfo Adv( NAME_AdvanceFrame, "", "", TEXT( "" ), EStatDataType::ST_int64, true, false );
+				GenerateNewFrame( FStatMessage( Adv.GetEncodedName(), EStatOperation::AdvanceFrameEventGameThread, 0ll, false ), Reader, bUseInAsync );
 				return true;
 			}
 
@@ -1100,27 +1102,7 @@ bool FServiceConnection::ReadAndConvertStatMessages( FArchive& Reader, bool bUse
 				{
 					if( Message.NameAndInfo.GetField<EStatOperation>() == EStatOperation::AdvanceFrameEventGameThread && ReadMessages > 2 )
 					{
-						AddCollectedStatMessages( Message );
-
-						// create an old format data frame from the data
-						GenerateProfilerDataFrame();
-
-						{
-							// add the frame to the work list
-							FScopeLock ScopeLock( &CriticalSection );
-							DataFrames.Add( CurrentData );
-							DataLoadingProgress = (double)Reader.Tell() / (double)Reader.TotalSize();
-						}
-
-						if( bUseInAsync )
-						{
-							if( DataFrames.Num() > FProfilerClientManager::MaxFramesPerTick )
-							{
-								FPlatformProcess::Sleep( 0.1f );
-							}
-						}
-
-
+						GenerateNewFrame( Message, Reader, bUseInAsync );
 					}
 
 					new (Messages)FStatMessage( Message );			
@@ -1137,6 +1119,7 @@ bool FServiceConnection::ReadAndConvertStatMessages( FArchive& Reader, bool bUse
 			}
 		}
 	}
+	// Obsolete. Remove later.
 	else
 	{
 		while( Reader.Tell() < Reader.TotalSize() )
@@ -1201,6 +1184,29 @@ bool FServiceConnection::ReadAndConvertStatMessages( FArchive& Reader, bool bUse
 	return false;
 }
 
+void FServiceConnection::GenerateNewFrame( FStatMessage Message, FArchive &Reader, bool bUseInAsync )
+{
+	AddCollectedStatMessages( Message );
+
+	// create an old format data frame from the data
+	GenerateProfilerDataFrame();
+
+	{
+		// add the frame to the work list
+		FScopeLock ScopeLock( &CriticalSection );
+		DataFrames.Add( CurrentData );
+		DataLoadingProgress = (double)Reader.Tell() / (double)Reader.TotalSize();
+	}
+
+	if (bUseInAsync)
+	{
+		if (DataFrames.Num() > FProfilerClientManager::MaxFramesPerTick)
+		{
+			FPlatformProcess::Sleep( 0.1f );
+		}
+	}
+}
+
 void FServiceConnection::AddCollectedStatMessages( FStatMessage Message )
 {
 	SCOPE_CYCLE_COUNTER( STAT_PC_AddStatMessages );
diff --git a/Engine/Source/Developer/ProfilerClient/Private/ProfilerClientManager.h b/Engine/Source/Developer/ProfilerClient/Private/ProfilerClientManager.h
--- a/Engine/Source/Developer/ProfilerClient/Private/ProfilerClientManager.h
+++ b/Engine/Source/Developer/ProfilerClient/Private/ProfilerClientManager.h
@@ -102,6 +102,9 @@ struct FServiceConnection
 	 */
 	bool ReadAndConvertStatMessages( FArchive& Reader, bool bUseInAsync );
 
+	/** Generates a new frame for the profiler. */
+	void GenerateNewFrame( FStatMessage Message, FArchive &Reader, bool bUseInAsync );
+
 	/** Adds all collected stat messages to the current stats thread state. */
 	void AddCollectedStatMessages( FStatMessage Message );
 
