Copying //UE4/Dev-Core to //UE4/Main

diff --git a/Engine/Source/Developer/Profiler/Private/Widgets/SDataGraph.cpp b/Engine/Source/Developer/Profiler/Private/Widgets/SDataGraph.cpp
--- a/Engine/Source/Developer/Profiler/Private/Widgets/SDataGraph.cpp
+++ b/Engine/Source/Developer/Profiler/Private/Widgets/SDataGraph.cpp
@@ -240,7 +240,7 @@ void SDataGraph::UpdateState()
 	if( GraphDesc )
 	{
 		// Check if we need to force time based view mode.
-		const bool bCanBeDisplayedAsMulti = GraphDesc->CombinedGraphDataSource->CanBeDisplayedAsMulti();
+		const bool bCanBeDisplayedAsMulti = false;
 		if( bCanBeDisplayedAsMulti )
 		{
 			ViewMode = EDataGraphViewModes::Time;
@@ -566,7 +566,7 @@ int32 SDataGraph::OnPaint( const FPaintArgs& Args, const FGeometry& AllottedGeom
 			const int32 FrameStartIndex = GraphOffset + AvgFrameRate - (GraphOffset % AvgFrameRate);
 			const int32 FrameEndIndex = FMath::Min( GraphOffset + NumVisiblePoints, NumDataPoints );
 
-			const bool bCanBeDisplayedAsMulti = GraphDesc->CombinedGraphDataSource->CanBeDisplayedAsMulti();
+			const bool bCanBeDisplayedAsMulti = false;
 			const IDataProviderRef DataProvider = (*GraphDataSource)->GetDataProvider();
 
 			for( int32 FrameIndex = FrameStartIndex; FrameIndex < FrameEndIndex; FrameIndex += AvgFrameRate )
@@ -1288,72 +1288,6 @@ void SDataGraph::ShowContextMenu(const FVector2D& ScreenSpacePosition, const FPo
 
 	const bool bShouldCloseWindowAfterMenuSelection = true;
 	FMenuBuilder MenuBuilder( bShouldCloseWindowAfterMenuSelection, ProfilerCommandList );
-
-	if( !FProfilerManager::GetSettings().bSingleInstanceMode )
-	{
-		MenuBuilder.BeginSection( "ProfilerInstances", LOCTEXT("ContextMenu_ProfilerInstances", "Profiler Instances") );
-		{
-			// @TODO: Add to FProfilerMenuBuilder
-			struct FProfilerSessionMenu
-			{
-				static void Build( FMenuBuilder& InMenuBuilder, const FGuid SessionInstanceID, const FEventGraphDataHandlerRef InEventGraphDataHandler )
-				{
-					const FProfilerCommands& Commands = FProfilerManager::GetCommands();
-					const FProfilerActionManager& ProfilerActionMgr = FProfilerManager::GetActionManager();
-
-					FProfilerMenuBuilder::AddMenuEntry
-					( 
-						InMenuBuilder, 
-						Commands.ToggleDataPreview, 
-						ProfilerActionMgr.ToggleDataPreview_Custom(SessionInstanceID) 
-					);
-
-					FProfilerMenuBuilder::AddMenuEntry
-					( 
-						InMenuBuilder, 
-						Commands.ToggleDataCapture, 
-						ProfilerActionMgr.ToggleDataCapture_Custom(SessionInstanceID) 
-					);
-
-					if( SessionInstanceID.IsValid() )
-					{
-						FProfilerMenuBuilder::AddMenuEntry
-						( 
-							InMenuBuilder, 
-							Commands.ToggleShowDataGraph, 
-							ProfilerActionMgr.ToggleShowDataGraph_Custom(SessionInstanceID) 
-						);
-					}
-				}
-			};
-
-			if( FProfilerManager::Get()->GetProfilerInstancesNum() > 1 )
-			{
-				MenuBuilder.AddSubMenu
-				( 
-					LOCTEXT("ContextMenu_AllProfilerInstances", "AllInstances"), 
-					LOCTEXT("ContextMenu_AllProfilerInstances_TT", "All profiler instances options"), 
-					FNewMenuDelegate::CreateStatic( &FProfilerSessionMenu::Build, FGuid(), EventGraphDataHandler )
-				);
-			}
-
-			MenuBuilder.AddMenuSeparator();
-
-			for( auto It = FProfilerManager::Get()->GetProfilerInstancesIterator(); It; ++It )
-			{
-				const FProfilerSessionRef& ProfilerSession = It.Value();
-				const FGuid SessionInstanceID = ProfilerSession->GetInstanceID();
-
-				MenuBuilder.AddSubMenu
-				( 
-					FText::FromString( ProfilerSession->GetShortName() ), 
-					LOCTEXT("ContextMenu_InstancesList_TT", "Profiler instance options"), 
-					FNewMenuDelegate::CreateStatic( &FProfilerSessionMenu::Build, SessionInstanceID, EventGraphDataHandler )
-				);	
-			}
-		}
-		MenuBuilder.EndSection();
-	}
 	
 	MenuBuilder.BeginSection( "ViewMode", LOCTEXT("ContextMenu_ViewMode", "View Mode") );
 	{
@@ -1363,16 +1297,6 @@ void SDataGraph::ShowContextMenu(const FVector2D& ScreenSpacePosition, const FPo
 	}
 	MenuBuilder.EndSection();
 
-	if( !FProfilerManager::GetSettings().bSingleInstanceMode )
-	{
-		MenuBuilder.BeginSection( "MultiMode", LOCTEXT("ContextMenu_MultiMode", "Multi Mode") );
-		{
-			MenuBuilder.AddMenuEntry( ProfilerCommands.DataGraph_MultiMode_SetCombined );
-			MenuBuilder.AddMenuEntry( ProfilerCommands.DataGraph_MultiMode_SetOneLinePerDataSource );
-		}
-		MenuBuilder.EndSection();
-	}
-
 	MenuBuilder.BeginSection( TEXT("Misc"), LOCTEXT("Miscellaneous", "Miscellaneous") );
 	{
 		MenuBuilder.AddMenuEntry( FProfilerManager::GetCommands().EventGraph_SelectAllFrames );
@@ -1409,24 +1333,6 @@ void SDataGraph::BindCommands()
 		FCanExecuteAction::CreateSP( this, &SDataGraph::ViewMode_SetTimeBased_CanExecute ),
 		FIsActionChecked::CreateSP( this, &SDataGraph::ViewMode_SetTimeBased_IsChecked )
 	);
- 
-	// DataGraph_MultiMode_SetCombined
-	ProfilerCommandList->MapAction
-	( 
-		ProfilerCommands.DataGraph_MultiMode_SetCombined,
-		FExecuteAction::CreateSP( this, &SDataGraph::MultiMode_SetCombined_Execute ),
-		FCanExecuteAction::CreateSP( this, &SDataGraph::MultiMode_SetCombined_CanExecute ),
-		FIsActionChecked::CreateSP( this, &SDataGraph::MultiMode_SetCombined_IsChecked )
-	);
- 
-	// DataGraph_MultiMode_SetCombined
-	ProfilerCommandList->MapAction
-	( 
-		ProfilerCommands.DataGraph_MultiMode_SetOneLinePerDataSource,
-		FExecuteAction::CreateSP( this, &SDataGraph::MultiMode_SetOneLinePerDataSource_Execute ),
-		FCanExecuteAction::CreateSP( this, &SDataGraph::MultiMode_SetOneLinePerDataSource_CanExecute ),
-		FIsActionChecked::CreateSP( this, &SDataGraph::MultiMode_SetOneLinePerDataSource_IsChecked )
-	);
 }
 
 /*-----------------------------------------------------------------------------
@@ -1473,47 +1379,6 @@ bool SDataGraph::ViewMode_SetTimeBased_IsChecked() const
 	return ViewMode == EDataGraphViewModes::Time;
 }
 
-/*-----------------------------------------------------------------------------
-	MultiMode_SetCombined
------------------------------------------------------------------------------*/
-
-void SDataGraph::MultiMode_SetCombined_Execute()
-{
-	MultiMode = EDataGraphMultiModes::Combined;
-	UpdateState();
-}
-
-bool SDataGraph::MultiMode_SetCombined_CanExecute() const
-{
-	const bool bCanBeDisplayedAsMulti = GetFirstGraph() ? GetFirstGraph()->CombinedGraphDataSource->CanBeDisplayedAsMulti() : false;
-	return MultiMode != EDataGraphMultiModes::Combined && bCanBeDisplayedAsMulti && ViewMode == EDataGraphViewModes::Time;
-}
-
-bool SDataGraph::MultiMode_SetCombined_IsChecked() const
-{
-	return MultiMode == EDataGraphMultiModes::Combined;
-}
-
-/*-----------------------------------------------------------------------------
-	MultiMode_SetOneLinePerDataSource
------------------------------------------------------------------------------*/
-
-void SDataGraph::MultiMode_SetOneLinePerDataSource_Execute()
-{
-	MultiMode = EDataGraphMultiModes::OneLinePerDataSource;
-	UpdateState();
-}
-
-bool SDataGraph::MultiMode_SetOneLinePerDataSource_CanExecute() const
-{
-	return MultiMode != EDataGraphMultiModes::OneLinePerDataSource;
-}
-
-bool SDataGraph::MultiMode_SetOneLinePerDataSource_IsChecked() const
-{
-	return MultiMode == EDataGraphMultiModes::OneLinePerDataSource;
-}
-
 void SDataGraph::EventGraph_OnRestoredFromHistory( uint32 FrameStartIndex, uint32 FrameEndIndex )
 {
 	UpdateState();
diff --git a/Engine/Source/Developer/Profiler/Private/Widgets/SDataGraph.h b/Engine/Source/Developer/Profiler/Private/Widgets/SDataGraph.h
--- a/Engine/Source/Developer/Profiler/Private/Widgets/SDataGraph.h
+++ b/Engine/Source/Developer/Profiler/Private/Widgets/SDataGraph.h
@@ -366,23 +366,6 @@ protected:
 	/** Handles FIsActionChecked for DataGraph_ViewMode_SetTimeBased. */
 	bool ViewMode_SetTimeBased_IsChecked() const;
 
-
-	/** Handles FExecuteAction for DataGraph_MultiMode_SetCombined. */
-	void MultiMode_SetCombined_Execute();
-	/** Handles FCanExecuteAction for DataGraph_MultiMode_SetCombined. */
-	bool MultiMode_SetCombined_CanExecute() const;
-	/** Handles FIsActionChecked for DataGraph_MultiMode_SetCombined. */
-	bool MultiMode_SetCombined_IsChecked() const;
-
-	
-	/** Handles UI Action's delegates for DataGraph_MultiMode_SetOneLinePerDataSource. */
-	void MultiMode_SetOneLinePerDataSource_Execute();
-	/** Handles FCanExecuteAction for DataGraph_MultiMode_SetOneLinePerDataSource. */
-	bool MultiMode_SetOneLinePerDataSource_CanExecute() const;
-	/** Handles FIsActionChecked for DataGraph_MultiMode_SetOneLinePerDataSource. */
-	bool MultiMode_SetOneLinePerDataSource_IsChecked() const;
-
-
 	const FGraphDescription* GetFirstGraph() const
 	{
 		const FGraphDescription* Result = StatIDToGraphDescriptionMapping.Num() > 0 ? &StatIDToGraphDescriptionMapping.CreateConstIterator().Value() : NULL;
diff --git a/Engine/Source/Developer/Profiler/Private/Widgets/SFiltersAndPresets.cpp b/Engine/Source/Developer/Profiler/Private/Widgets/SFiltersAndPresets.cpp
--- a/Engine/Source/Developer/Profiler/Private/Widgets/SFiltersAndPresets.cpp
+++ b/Engine/Source/Developer/Profiler/Private/Widgets/SFiltersAndPresets.cpp
@@ -43,12 +43,7 @@ public:
 	SFiltersAndPresetsTooltip( const uint32 InStatID )
 		: StatID( InStatID )
 	{
-		// TODO: At this moment only single profiler instance is supported.
-		const int32 NumInstances = FProfilerManager::Get()->GetProfilerInstancesNum();
-		if( NumInstances == 1 )
-		{
-			ProfilerSession = FProfilerManager::Get()->GetProfilerInstancesIterator().Value();
-		}
+		ProfilerSession = FProfilerManager::Get()->GetProfilerSession();
 	}
 
 	BEGIN_SLATE_FUNCTION_BUILD_OPTIMIZATION
@@ -763,8 +758,8 @@ END_SLATE_FUNCTION_BUILD_OPTIMIZATION
 
 void SFiltersAndPresets::ProfilerManager_OnRequestFilterAndPresetsUpdate()
 {
-	auto It = FProfilerManager::Get()->GetProfilerInstancesIterator();
-	UpdateGroupAndStatTree(It.Value());
+	FProfilerSessionPtr ProfilerSessionLocal = FProfilerManager::Get()->GetProfilerSession();
+	UpdateGroupAndStatTree( ProfilerSessionLocal );
 }
 
 void SFiltersAndPresets::UpdateGroupAndStatTree( const FProfilerSessionPtr InProfilerSession )
diff --git a/Engine/Source/Developer/Profiler/Private/Widgets/SProfilerGraphPanel.cpp b/Engine/Source/Developer/Profiler/Private/Widgets/SProfilerGraphPanel.cpp
--- a/Engine/Source/Developer/Profiler/Private/Widgets/SProfilerGraphPanel.cpp
+++ b/Engine/Source/Developer/Profiler/Private/Widgets/SProfilerGraphPanel.cpp
@@ -20,7 +20,6 @@ SProfilerGraphPanel::~SProfilerGraphPanel()
 	if( ProfilerManager.IsValid() )
 	{
 		ProfilerManager->OnTrackedStatChanged().RemoveAll( this );
-		ProfilerManager->OnSessionInstancesUpdated().RemoveAll( this );
 		ProfilerManager->OnViewModeChanged().RemoveAll( this );
 
 		DataGraph->OnSelectionChangedForIndex().RemoveAll( ProfilerManager.Get() );
diff --git a/Engine/Source/Developer/Profiler/Private/Widgets/SProfilerSettings.cpp b/Engine/Source/Developer/Profiler/Private/Widgets/SProfilerSettings.cpp
--- a/Engine/Source/Developer/Profiler/Private/Widgets/SProfilerSettings.cpp
+++ b/Engine/Source/Developer/Profiler/Private/Widgets/SProfilerSettings.cpp
@@ -18,15 +18,6 @@ void SProfilerSettings::Construct( const FArguments& InArgs )
 	AddSeparator( SettingsGrid, CurrentRowPos );
 	AddHeader( LOCTEXT("MiscTitle","Miscellaneous"), SettingsGrid, CurrentRowPos );
 	AddOption
-	( 
-		LOCTEXT("bSingleInstanceMode_T","Single Instance Mode"), 
-		LOCTEXT("bSingleInstanceMode_TT","If True, the profiler will work in the single instance mode, all functionality related to multi instances will be disabled or removed from the UI"), 
-		SettingPtr->bSingleInstanceMode, 
-		SettingPtr->GetDefaults().bSingleInstanceMode,
-		SettingsGrid, 
-		CurrentRowPos 
-	);
-	AddOption
 	( 
 		LOCTEXT("bShowCoalescedViewModesInEventGraph_T","Show Coalesced View Modes In Event Graph"), 
 		LOCTEXT("bShowCoalescedViewModesInEventGraph_TT","If True, coalesced view modes related functionality will be added to the event graph"), 
@@ -214,10 +205,9 @@ FReply SProfilerSettings::SaveAndClose_OnClicked()
 
 FReply SProfilerSettings::ResetToDefaults_OnClicked()
 {
-	const FProfilerSettings& Defauls = SettingPtr->GetDefaults();
+	const FProfilerSettings& Defaults = SettingPtr->GetDefaults();
 
-	SettingPtr->bSingleInstanceMode = Defauls.bSingleInstanceMode;
-	SettingPtr->bShowCoalescedViewModesInEventGraph = Defauls.bShowCoalescedViewModesInEventGraph;
+	SettingPtr->bShowCoalescedViewModesInEventGraph = Defaults.bShowCoalescedViewModesInEventGraph;
 
 	return FReply::Handled();
 }
diff --git a/Engine/Source/Developer/Profiler/Private/Widgets/SProfilerToolbar.cpp b/Engine/Source/Developer/Profiler/Private/Widgets/SProfilerToolbar.cpp
--- a/Engine/Source/Developer/Profiler/Private/Widgets/SProfilerToolbar.cpp
+++ b/Engine/Source/Developer/Profiler/Private/Widgets/SProfilerToolbar.cpp
@@ -77,21 +77,7 @@ void SProfilerToolbar::Construct( const FArguments& InArgs )
 				ToolbarBuilder.MakeWidget()
 			]
 		]
-
-		+SHorizontalBox::Slot()
-		.AutoWidth()
-		.HAlign(HAlign_Right)
-		.VAlign(VAlign_Fill)
-		.Padding(0.0f)
-		[
-			SAssignNew(Border,SBorder)
-			.BorderImage( FEditorStyle::GetBrush("ToolPanel.GroupBorder") )
-			.Padding( 2.0f )
-			.Visibility( this, &SProfilerToolbar::IsInstancesOwnerVisible )
-		]
 	];
-
-	FProfilerManager::Get()->OnSessionInstancesUpdated().AddSP( this, &SProfilerToolbar::ProfilerManager_SessionsUpdated );
 }
 
 void SProfilerToolbar::ShowStats()
@@ -139,11 +125,9 @@ void DisplayFPSChart( const TSharedRef<FFPSAnalyzer> InFPSAnalyzer )
 
 void SProfilerToolbar::ShowFPSChart()
 {
-	for( auto It = FProfilerManager::Get()->GetProfilerInstancesIterator(); It; ++It )
+	const FProfilerSessionPtr ProfilerSession = FProfilerManager::Get()->GetProfilerSession();
+	if (ProfilerSession.IsValid())
 	{
-		const FProfilerSessionRef& ProfilerSession = It.Value();
-		const FGuid SessionInstanceID = ProfilerSession->GetInstanceID();
-
 		DisplayFPSChart( ProfilerSession->FPSAnalyzer );
 	}
 }
@@ -180,100 +164,4 @@ void SProfilerToolbar::CreateCommands()
 		);
 }
 
-EVisibility SProfilerToolbar::IsInstancesOwnerVisible() const
-{
-	return FProfilerManager::GetSettings().bSingleInstanceMode ? EVisibility::Collapsed : EVisibility::Visible;
-}
-
-void SProfilerToolbar::ProfilerManager_SessionsUpdated()
-{
-	TSharedRef<SHorizontalBox> SessionsHBox = SNew(SHorizontalBox);
-
-	if( FProfilerManager::Get()->GetProfilerInstancesNum() > 1 )
-	{
-		AddSessionInstanceItem( SessionsHBox, LOCTEXT("Global", "Global").ToString(), FGuid() );
-	}
-
-	for( auto It = FProfilerManager::Get()->GetProfilerInstancesIterator(); It; ++It )
-	{
-		const FProfilerSessionRef& ProfilerSession = It.Value();
-		const FGuid SessionInstanceID = ProfilerSession->GetInstanceID();
-
-		AddSessionInstanceItem( SessionsHBox, ProfilerSession->GetShortName(), SessionInstanceID );
-	}
-
-	Border->SetContent( SessionsHBox );
-}
-
-void SProfilerToolbar::AddSessionInstanceItem( TSharedRef<SHorizontalBox>& SessionsHBox, const FString& ProfilerSessionName, const FGuid& SessionInstanceID )
-{
-	SessionsHBox->AddSlot()
-	.HAlign(HAlign_Fill)
-	.VAlign(VAlign_Fill)
-	.AutoWidth()
-	.Padding(2.0f)
-	[
-		SNew(SComboButton)	
-		.HAlign(HAlign_Center)
-		.VAlign(VAlign_Center)
-
-		.ButtonContent()
-		[
-			SNew(STextBlock)
-			.Text( FText::FromString(ProfilerSessionName) )
-			.ShadowOffset( FVector2D(1.0f,1.0f) )
-			.ShadowColorAndOpacity( FColorList::SkyBlue )
-		]
-
-		.MenuContent()
-		[
-			BuildProfilerSessionContextMenu( SessionInstanceID )
-		]
-	];
-}
-
-TSharedRef<SWidget> SProfilerToolbar::BuildProfilerSessionContextMenu( const FGuid SessionInstanceID )
-{
-	TSharedPtr<FUICommandList> ProfilerCommandList = FProfilerManager::Get()->GetCommandList();
-	const FProfilerCommands& ProfilerCommands = FProfilerManager::GetCommands();
-	const FProfilerActionManager& ProfilerActionManager = FProfilerManager::GetActionManager();
-
-	const bool bShouldCloseWindowAfterMenuSelection = true;
-	FMenuBuilder MenuBuilder( bShouldCloseWindowAfterMenuSelection, ProfilerCommandList );
-
-	MenuBuilder.BeginSection( "Menu_Instance", LOCTEXT("SessionInstance", "Session Instance") );
-	MenuBuilder.EndSection();
-	
-	MenuBuilder.BeginSection( "Menu_General", LOCTEXT("General", "General") );
-	{
-		FProfilerMenuBuilder::AddMenuEntry
-		( 
-			MenuBuilder, 
-			ProfilerCommands.ToggleDataPreview, 
-			ProfilerActionManager.ToggleDataPreview_Custom(SessionInstanceID) 
-		);
-
-		FProfilerMenuBuilder::AddMenuEntry
-		( 
-			MenuBuilder, 
-			ProfilerCommands.ToggleDataCapture, 
-			ProfilerActionManager.ToggleDataCapture_Custom(SessionInstanceID) 
-		);
-
-		if( SessionInstanceID.IsValid() )
-		{
-			FProfilerMenuBuilder::AddMenuEntry
-			( 
-				MenuBuilder, 
-				ProfilerCommands.ToggleShowDataGraph, 
-				ProfilerActionManager.ToggleShowDataGraph_Custom(SessionInstanceID) 
-			);
-		}
-	}
-	MenuBuilder.EndSection();
-
-	TSharedRef<SWidget> MenuWidget = MenuBuilder.MakeWidget();
-	return MenuWidget;
-}
-
 #undef LOCTEXT_NAMESPACE
diff --git a/Engine/Source/Developer/Profiler/Private/Widgets/SProfilerToolbar.h b/Engine/Source/Developer/Profiler/Private/Widgets/SProfilerToolbar.h
--- a/Engine/Source/Developer/Profiler/Private/Widgets/SProfilerToolbar.h
+++ b/Engine/Source/Developer/Profiler/Private/Widgets/SProfilerToolbar.h
@@ -32,17 +32,6 @@ protected:
 	/** Called when the list of sessions has changed. @see FProfilerManager.FOnSessionsUpdated */
 	void ProfilerManager_SessionsUpdated();
 
-	void AddSessionInstanceItem( TSharedRef<SHorizontalBox>& SessionsHBox, const FString& ProfilerSessionName, const FGuid& SessionInstanceID );
-	EVisibility IsInstancesOwnerVisible() const;
-
-	/**
-	 * Builds the context menu for the specified profiler session.
-	 *
-	 * @param SessionInstanceID - the session instance ID for which the context menu will be built 
-	 *
-	 */
-	TSharedRef<SWidget> BuildProfilerSessionContextMenu( const FGuid SessionInstanceID );
-
 	/** Create the UI commands for the toolbar */
 	void CreateCommands();
 
diff --git a/Engine/Source/Developer/Profiler/Private/Widgets/SProfilerWindow.cpp b/Engine/Source/Developer/Profiler/Private/Widgets/SProfilerWindow.cpp
--- a/Engine/Source/Developer/Profiler/Private/Widgets/SProfilerWindow.cpp
+++ b/Engine/Source/Developer/Profiler/Private/Widgets/SProfilerWindow.cpp
@@ -351,7 +351,7 @@ void SProfilerWindow::UpdateEventGraph( const FGuid ProfilerInstanceID, const FE
 
 EVisibility SProfilerWindow::IsSessionOverlayVissible() const
 {
-	if( FProfilerManager::Get()->GetProfilerInstancesNum() > 0 )
+	if( FProfilerManager::Get()->HasValidSession() )
 	{
 		return EVisibility::Hidden;
 	}
diff --git a/Engine/Source/Developer/ProfilerClient/Private/ProfilerClientManager.cpp b/Engine/Source/Developer/ProfilerClient/Private/ProfilerClientManager.cpp
--- a/Engine/Source/Developer/ProfilerClient/Private/ProfilerClientManager.cpp
+++ b/Engine/Source/Developer/ProfilerClient/Private/ProfilerClientManager.cpp
@@ -48,9 +48,6 @@ FProfilerClientManager::FProfilerClientManager( const IMessageBusRef& InMessageB
 	LastPingTime = FDateTime::Now();
 	RetryTime = 5.f;
 
-#if PROFILER_THREADED_LOAD
-	LoadTask = nullptr;
-#endif
 	LoadConnection = nullptr;
 	MessageDelegateHandle = FTicker::GetCoreTicker().AddTicker(MessageDelegate, 0.1f);
 #endif
@@ -59,21 +56,7 @@ FProfilerClientManager::FProfilerClientManager( const IMessageBusRef& InMessageB
 FProfilerClientManager::~FProfilerClientManager()
 {
 #if STATS
-	// Delete all active file writers and remove temporary files.
-	for( auto It = ActiveTransfers.CreateIterator(); It; ++It )
-	{
-		FReceivedFileInfo& ReceivedFileInfo = It.Value();
-		
-		delete ReceivedFileInfo.FileWriter;
-		ReceivedFileInfo.FileWriter = nullptr;
-
-		IFileManager::Get().Delete( *ReceivedFileInfo.DestFilepath );
-
-		UE_LOG( LogProfilerClient, Log, TEXT( "File service-client transfer aborted: %s" ), *It.Key() );
-	}
-
-	FTicker::GetCoreTicker().RemoveTicker(MessageDelegateHandle);
-	FTicker::GetCoreTicker().RemoveTicker(TickDelegateHandle);
+	Shutdown();
 
 	Unsubscribe();
 
@@ -137,34 +120,6 @@ void FProfilerClientManager::Track( const FGuid& Instance )
 #endif
 }
 
-void FProfilerClientManager::Track( const TArray<TSharedPtr<ISessionInstanceInfo>>& Instances )
-{
-#if STATS
-	if (MessageEndpoint.IsValid() && ActiveSessionId.IsValid())
-	{
-		TArray<FGuid> ActiveInstances;
-		Connections.GenerateKeyArray(ActiveInstances);
-
-		for (int32 i = 0; i < Instances.Num(); ++i)
-		{
-			if (Connections.Find(Instances[i]->GetInstanceId()) == nullptr)
-			{
-				Track(Instances[i]->GetInstanceId());
-			}
-			else
-			{
-				ActiveInstances.Remove(Instances[i]->GetInstanceId());
-			}
-		}
-
-		for (int32 i = 0; i < ActiveInstances.Num(); ++i)
-		{
-			Untrack(ActiveInstances[i]);
-		}
-	}
-#endif
-}
-
 void FProfilerClientManager::Untrack( const FGuid& Instance )
 {
 #if STATS
@@ -246,76 +201,117 @@ void FProfilerClientManager::SetPreviewState( const bool bRequestedPreviewState,
 #endif
 }
 
-void FProfilerClientManager::LoadCapture( const FString& DataFilepath, const FGuid& ProfileId )
-{
-#if STATS
-	UE_LOG( LogStats, Warning, TEXT( "Started loading: %s" ), *DataFilepath );
+/*-----------------------------------------------------------------------------
+	New read test, still temporary, but around 4x faster
+-----------------------------------------------------------------------------*/
 
-	// start an async load
-	LoadConnection = &Connections.FindOrAdd(ProfileId);
-	LoadConnection->InstanceId = ProfileId;
-	LoadConnection->MetaData.CriticalSection = &LoadConnection->CriticalSection;
-	LoadConnection->MetaData.SecondsPerCycle = FPlatformTime::GetSecondsPerCycle(); // fix this by adding a message which specifies this
+class FNewStatsReader : public FStatsReadFile
+{
+	friend struct FStatsReader<FNewStatsReader>;
+	typedef FStatsReadFile Super;
 
-	const int64 Size = IFileManager::Get().FileSize( *DataFilepath );
-	if( Size < 4 )
+public:
+	/** Initialize. */
+	void Initialize( FProfilerClientManager* InProfilerClientManager, FServiceConnection* InLoadConnection )
 	{
-		UE_LOG( LogProfilerClient, Error, TEXT( "Could not open: %s" ), *DataFilepath );
-		return;
+		ProfilerClientManager = InProfilerClientManager;
+		LoadConnection = InLoadConnection;
 	}
 
-#if PROFILER_THREADED_LOAD
-	FArchive* FileReader = IFileManager::Get().CreateFileReader(*DataFilepath);
-#else
-	FileReader = IFileManager::Get().CreateFileReader( *DataFilepath );
-#endif
+protected:
 
-	if( !FileReader )
+	/** Initialization constructor. */
+	FNewStatsReader( const TCHAR* InFilename )
+		: FStatsReadFile( InFilename, false )
+		, ProfilerClientManager( nullptr )
+		, LoadConnection( nullptr )
 	{
-		UE_LOG( LogProfilerClient, Error, TEXT( "Could not open: %s" ), *DataFilepath );
-		return;
-	}
+		// Keep only the last frame.
+		SetHistoryFrames( 1 );
+	}	
 
-	if( !LoadConnection->Stream.ReadHeader( *FileReader ) )
+	/** Called every each frame has been read from the file. */
+	virtual void ReadStatsFrame( const TArray<FStatMessage>& CondensedMessages, const int64 Frame ) override
 	{
-		UE_LOG( LogProfilerClient, Error, TEXT( "Could not open, bad magic: %s" ), *DataFilepath );
-		delete FileReader;
-		return;
-	}
+		SCOPE_CYCLE_COUNTER( STAT_PC_GenerateDataFrame );
 
-	// This shouldn't happen.
-	if( LoadConnection->Stream.Header.bRawStatsFile )
-	{
-		delete FileReader;
-		return;
+		FProfilerDataFrame& DataFrame = LoadConnection->CurrentData;
+
+		DataFrame.Frame = Frame;
+		DataFrame.FrameStart = 0.0;
+		DataFrame.CountAccumulators.Reset();
+		DataFrame.CycleGraphs.Reset();
+		DataFrame.FloatAccumulators.Reset();
+		DataFrame.MetaDataUpdated = false;
+
+		// Get the stat stack root and the non frame stats.
+		FRawStatStackNode Stack;
+		TArray<FStatMessage> NonFrameStats;
+		State.UncondenseStackStats( CondensedMessages, Stack, nullptr, &NonFrameStats );
+
+		LoadConnection->GenerateCycleGraphs( Stack, DataFrame.CycleGraphs );
+		LoadConnection->GenerateAccumulators( NonFrameStats, DataFrame.CountAccumulators, DataFrame.FloatAccumulators );
+
+		// Create a copy of the stats metadata.
+		FStatMetaData* MetaDataPtr = nullptr;
+
+		if (DataFrame.MetaDataUpdated)
+		{
+			MetaDataPtr = new FStatMetaData();
+			*MetaDataPtr = LoadConnection->StatMetaData;
+		}
+
+		// Create a copy of the stats data.
+		FProfilerDataFrame* DataFramePtr = new FProfilerDataFrame();
+		*DataFramePtr = DataFrame;
+
+		// Send to game thread.
+		FSimpleDelegateGraphTask::CreateAndDispatchWhenReady
+		(
+			FSimpleDelegateGraphTask::FDelegate::CreateRaw( ProfilerClientManager, &FProfilerClientManager::SendProfilerDataFrameToGame, DataFramePtr, MetaDataPtr ),
+			TStatId(), nullptr, ENamedThreads::GameThread
+		);
 	}
 
-	const bool bIsFinalized = LoadConnection->Stream.Header.IsFinalized();
-	if( bIsFinalized )
+	/** Called after reading all data from the file. */
+	virtual void PreProcessStats() override
 	{
-		// Read metadata.
-		TArray<FStatMessage> MetadataMessages;
-		LoadConnection->Stream.ReadFNamesAndMetadataMessages( *FileReader, MetadataMessages );
-		LoadConnection->CurrentThreadState.ProcessMetaDataOnly( MetadataMessages );
-
-		// Read frames offsets.
-		LoadConnection->Stream.ReadFramesOffsets( *FileReader );
-		FileReader->Seek( LoadConnection->Stream.FramesInfo[0].FrameFileOffset );
+		FFunctionGraphTask::CreateAndDispatchWhenReady( [=]()
+		{
+			ProfilerClientManager->FinalizeLoading();
+		}
+		, TStatId(), NULL, ENamedThreads::GameThread );
 	}
 
-	if( LoadConnection->Stream.Header.HasCompressedData() )
+	FProfilerClientManager* ProfilerClientManager;
+	FServiceConnection* LoadConnection;
+};
+
+void FServiceConnection::LoadCapture( const FString& DataFilepath, FProfilerClientManager* ProfilerClientManager )
+{
+	StatsReader = FStatsReader<FNewStatsReader>::Create( *DataFilepath );
+	if (StatsReader)
 	{
-		UE_CLOG( !bIsFinalized, LogProfilerClient, Fatal, TEXT( "Compressed stats file has to be finalized" ) );
+		StatsReader->Initialize( ProfilerClientManager, this );
+		StatsReader->ReadAndProcessAsynchronously();
 	}
+}
+
 
-#if PROFILER_THREADED_LOAD
-	LoadTask = new FAsyncTask<FAsyncReadWorker>(LoadConnection, FileReader);
-	LoadTask->StartBackgroundTask();
-#endif
+
+void FProfilerClientManager::LoadCapture( const FString& DataFilepath, const FGuid& ProfileId )
+{
+#if STATS
+	// Start an async load.
+	LoadConnection = &Connections.FindOrAdd(ProfileId);
+	LoadConnection->InstanceId = ProfileId;
+	LoadConnection->StatMetaData.SecondsPerCycle = FPlatformTime::GetSecondsPerCycle(); // fix this by adding a message which specifies this
+
+	ProfilerLoadStartedDelegate.Broadcast( ProfileId );
+	LoadConnection->LoadCapture( DataFilepath, this );
 
 	RetryTime = 0.05f;
-	TickDelegateHandle = FTicker::GetCoreTicker().AddTicker(TickDelegate, RetryTime);
-	ProfilerLoadStartedDelegate.Broadcast(ProfileId);
+	TickDelegateHandle = FTicker::GetCoreTicker().AddTicker(TickDelegate, RetryTime);	
 #endif
 }
 
@@ -352,7 +348,7 @@ void FProfilerClientManager::RequestLastCapturedFile( const FGuid& InstanceId /*
 void FProfilerClientManager::HandleMessageBusShutdown()
 {
 #if STATS
-	Unsubscribe();
+	Shutdown();
 
 	MessageEndpoint.Reset();
 	MessageBus.Reset();
@@ -369,7 +365,7 @@ void FProfilerClientManager::HandleServiceAuthorizeMessage( const FProfilerServi
 		Connection.Initialize(Message, Context);
 
 		// Fire a meta data update message
-		ProfilerMetaDataUpdatedDelegate.Broadcast(Message.InstanceId);
+		//ProfilerMetaDataUpdatedDelegate.Broadcast(Message.InstanceId);
 
 		// Fire the client connection event
 		ProfilerClientConnectedDelegate.Broadcast(ActiveSessionId, Message.InstanceId);
@@ -379,41 +375,42 @@ void FProfilerClientManager::HandleServiceAuthorizeMessage( const FProfilerServi
 #endif
 }
 
-void FServiceConnection::Initialize( const FProfilerServiceAuthorize& Message, const IMessageContextRef& Context )
-{
-#if STATS
-	ProfilerServiceAddress = Context->GetSender();
-	InstanceId = Message.InstanceId;
-	CurrentData.Frame = 0;
-
-	// Add the supplied meta data.
-	FMemoryReader MemoryReader( Message.Data );
+/*-----------------------------------------------------------------------------
+	FServiceConnection
+-----------------------------------------------------------------------------*/
 
-	MetaData.CriticalSection = &CriticalSection;
-	int64 Size = MemoryReader.TotalSize();
+FServiceConnection::FServiceConnection() 
+	: StatsReader( nullptr )
+{
+}
 
-	const bool bVerifyHeader = Stream.ReadHeader( MemoryReader );
-	//check( bVerifyHeader );
+FServiceConnection::FServiceConnection( const FServiceConnection& InConnnection )
+{
+}
 
-	// Read in the data.
-	TArray<FStatMessage> StatMessages;
+FServiceConnection::~FServiceConnection()
+{
+	if (StatsReader)
 	{
-		SCOPE_CYCLE_COUNTER(STAT_PC_ReadStatMessages);
-		while(MemoryReader.Tell() < Size)
+		StatsReader->RequestStop();
+
+		while (StatsReader->IsBusy())
 		{
-			// read the message
-			new (StatMessages)FStatMessage( Stream.ReadMessage( MemoryReader ) );
+			FPlatformProcess::Sleep( 2.0f );
+			UE_LOG( LogStats, Log, TEXT( "RequestStop: Stage: %s / %3i%%" ), *StatsReader->GetProcessingStageAsString(), StatsReader->GetStageProgress() );
 		}
-		new (StatMessages) FStatMessage(FStatConstants::AdvanceFrame.GetEncodedName(), EStatOperation::AdvanceFrameEventGameThread, 1LL, false);
-	}
 
-	// generate a thread state from the data
-	{
-		SCOPE_CYCLE_COUNTER(STAT_PC_AddStatMessages);
-		CurrentThreadState.AddMessages(StatMessages);
+		delete StatsReader;
+		StatsReader = nullptr;
 	}
+}
 
-	UpdateMetaData();
+void FServiceConnection::Initialize( const FProfilerServiceAuthorize& Message, const IMessageContextRef& Context )
+{
+#if STATS
+	ProfilerServiceAddress = Context->GetSender();
+	InstanceId = Message.InstanceId;
+	CurrentData.Frame = 0;
 #endif
 }
 
@@ -553,36 +550,6 @@ void FProfilerClientManager::HandleServicePingMessage( const FProfilerServicePin
 #endif
 }
 
-#if PROFILER_THREADED_LOAD
-bool FProfilerClientManager::AsyncLoad()
-{
-#if STATS
-	BroadcastMetadataUpdate();
-
-	if (LoadTask->IsDone())
-	{
-		FinalizeLoading();
-		return false;
-	}
-#endif
-	return true;
-}
-#else
-bool FProfilerClientManager::SyncLoad()
-{
-#if STATS
-	const bool bFinalize = LoadConnection->ReadAndConvertStatMessages( *FileReader, false );
-	BroadcastMetadataUpdate();
-
-	if( bFinalize )
-	{
-		FinalizeLoading();
-		return false;
-	}
-#endif
-	return true;
-}
-#endif
 
 bool FProfilerClientManager::HandleTicker( float DeltaTime )
 {
@@ -600,14 +567,6 @@ bool FProfilerClientManager::HandleTicker( float DeltaTime )
 		}
 		LastPingTime = FDateTime::Now();
 	}
-	else if (LoadConnection)
-	{
-#if PROFILER_THREADED_LOAD
-		return AsyncLoad();
-#else
-		return SyncLoad();
-#endif
-	}
 #endif
 	return false;
 }
@@ -615,6 +574,7 @@ bool FProfilerClientManager::HandleTicker( float DeltaTime )
 bool FProfilerClientManager::HandleMessagesTicker( float DeltaTime )
 {
 #if STATS
+	// #YRX_tag 2015-11-26 rework 
 	for (auto It = Connections.CreateIterator(); It; ++It)
 	{
 		FServiceConnection& Connection = It.Value();
@@ -634,6 +594,8 @@ bool FProfilerClientManager::HandleMessagesTicker( float DeltaTime )
 				break;
 			}
 
+			FScopeLogTime SLT( "HandleMessagesTicker" );
+
 			int64 FrameNum = Frames[Index];
 			TArray<uint8>& Data = *Connection.PendingMessages.Find( FrameNum );
 
@@ -658,7 +620,7 @@ bool FProfilerClientManager::HandleMessagesTicker( float DeltaTime )
 						Connection.AddCollectedStatMessages( Message );
 					}
 
-					new (Connection.Messages) FStatMessage(Message);
+					new (Connection.PendingStatMessagesMessages) FStatMessage(Message);
 				}
 			}
 
@@ -668,11 +630,11 @@ bool FProfilerClientManager::HandleMessagesTicker( float DeltaTime )
 			// Fire a meta data update message
 			if (Connection.CurrentData.MetaDataUpdated)
 			{
-				ProfilerMetaDataUpdatedDelegate.Broadcast(Connection.InstanceId);
+				ProfilerMetaDataUpdatedDelegate.Broadcast( Connection.InstanceId, Connection.StatMetaData );
 			}
 
 			// send the data out
-			ProfilerDataDelegate.Broadcast( Connection.InstanceId, Connection.CurrentData, 0.0f );
+			ProfilerDataDelegate.Broadcast( Connection.InstanceId, Connection.CurrentData );
 
 			Connection.PendingMessages.Remove( FrameNum );
 		}
@@ -705,7 +667,7 @@ bool FProfilerClientManager::HandleMessagesTicker( float DeltaTime )
 void FProfilerClientManager::HandleServicePreviewAckMessage( const FProfilerServicePreviewAck& Message, const IMessageContextRef& Context )
 {
 #if STATS
-	if (ActiveSessionId.IsValid() && Connections.Find(Message.InstanceId) != NULL)
+	if (ActiveSessionId.IsValid() && Connections.Find(Message.InstanceId) != nullptr)
 	{
 		FServiceConnection& Connection = *Connections.Find(Message.InstanceId);
 
@@ -756,14 +718,14 @@ void FProfilerClientManager::DecompressDataAndSendToGame( FProfilerServiceData2*
 
 	FSimpleDelegateGraphTask::CreateAndDispatchWhenReady
 	(
-		FSimpleDelegateGraphTask::FDelegate::CreateRaw( this, &FProfilerClientManager::SendToGame, DateToGame, ToProcess->Frame, ToProcess->InstanceId ), 
+		FSimpleDelegateGraphTask::FDelegate::CreateRaw( this, &FProfilerClientManager::SendDataToGame, DateToGame, ToProcess->Frame, ToProcess->InstanceId ), 
 		TStatId(), nullptr, ENamedThreads::GameThread
 	);
 
 	delete ToProcess;
 }
 
-void FProfilerClientManager::SendToGame( TArray<uint8>* DataToGame, int64 Frame, const FGuid InstanceId )
+void FProfilerClientManager::SendDataToGame( TArray<uint8>* DataToGame, int64 Frame, const FGuid InstanceId )
 {
 	if (ActiveSessionId.IsValid() && Connections.Find( InstanceId ) != nullptr)
 	{
@@ -777,92 +739,94 @@ void FProfilerClientManager::SendToGame( TArray<uint8>* DataToGame, int64 Frame,
 	delete DataToGame;
 }
 
-void FProfilerClientManager::BroadcastMetadataUpdate()
+// #YRX_tag 2015-11-19 Convert to thread safe pointers?
+void FProfilerClientManager::SendProfilerDataFrameToGame( FProfilerDataFrame* NewData, FStatMetaData* MetaDataPtr )
 {
-	while( LoadConnection->DataFrames.Num() > 0 )
+	if (MetaDataPtr)
 	{
-		FScopeLock ScopeLock( &(LoadConnection->CriticalSection) );
+		ProfilerMetaDataUpdatedDelegate.Broadcast( LoadConnection->InstanceId, *MetaDataPtr );
+		delete MetaDataPtr;
+		MetaDataPtr = nullptr;
+	}
 
-		// Fire a meta data update message
-		if( LoadConnection->DataFrames[0].MetaDataUpdated )
-		{
-			ProfilerMetaDataUpdatedDelegate.Broadcast( LoadConnection->InstanceId );
-			ProfilerLoadedMetaDataDelegate.Broadcast( LoadConnection->InstanceId );
-		}
+	if (NewData)
+	{
+		ProfilerDataDelegate.Broadcast( LoadConnection->InstanceId, *NewData );
+		delete NewData;
+		NewData = nullptr;
+	}
+}
+
+void FProfilerClientManager::Shutdown()
+{
+	// Delete all active file writers and remove temporary files.
+	for (auto It = ActiveTransfers.CreateIterator(); It; ++It)
+	{
+		FReceivedFileInfo& ReceivedFileInfo = It.Value();
+
+		delete ReceivedFileInfo.FileWriter;
+		ReceivedFileInfo.FileWriter = nullptr;
+
+		IFileManager::Get().Delete( *ReceivedFileInfo.DestFilepath );
 
-		FProfilerDataFrame& DataFrame = LoadConnection->DataFrames[0];
-		ProfilerDataDelegate.Broadcast( LoadConnection->InstanceId, DataFrame, LoadConnection->DataLoadingProgress );
-		LoadConnection->DataFrames.RemoveAt( 0 );
+		UE_LOG( LogProfilerClient, Log, TEXT( "File service-client transfer aborted: %s" ), *It.Key() );
 	}
+
+	FTicker::GetCoreTicker().RemoveTicker( MessageDelegateHandle );
+	FTicker::GetCoreTicker().RemoveTicker( TickDelegateHandle );
 }
 
 void FProfilerClientManager::FinalizeLoading()
 {
 	ProfilerLoadCompletedDelegate.Broadcast( LoadConnection->InstanceId );
+	delete LoadConnection->StatsReader;
+	LoadConnection->StatsReader = nullptr;
 	LoadConnection = nullptr;
-#if	PROFILER_THREADED_LOAD
-	delete LoadTask;
-	LoadTask = nullptr;
-#else
-	delete FileReader;
-	FileReader = nullptr;
-#endif // PROFILER_THREADED_LOAD
-	RetryTime = 5.f;
-}
 
-void FProfilerClientManager::FAsyncReadWorker::DoWork()
-{
-#if STATS
-	const bool bFinalize = LoadConnection->ReadAndConvertStatMessages( *FileReader, true );
-	if( bFinalize )
-	{
-		FileReader->Close();
-		delete FileReader;
-	}
-#endif
+	RetryTime = 5.f;
 }
 
 #if STATS
-void FServiceConnection::UpdateMetaData()
+void FServiceConnection::UpdateMetaData_DISABLED()
 {
-	// loop through the stats meta data messages
-	for (auto It = CurrentThreadState.ShortNameToLongName.CreateConstIterator(); It; ++It)
-	{
-		FStatMessage const& LongName = It.Value();
-		const FName GroupName = LongName.NameAndInfo.GetGroupName();
-
-		uint32 StatType = STATTYPE_Error;
-		if (LongName.NameAndInfo.GetField<EStatDataType>() == EStatDataType::ST_int64)
-		{
-			if( LongName.NameAndInfo.GetFlag( EStatMetaFlags::IsCycle ) )
-			{
-				StatType = STATTYPE_CycleCounter;
-			}
-			else if( LongName.NameAndInfo.GetFlag( EStatMetaFlags::IsMemory ) )
-			{
-				StatType = STATTYPE_MemoryCounter;
-			}
-			else
-			{
-				StatType = STATTYPE_AccumulatorDWORD;
-			}
-		}
-		else if (LongName.NameAndInfo.GetField<EStatDataType>() == EStatDataType::ST_double)
-		{
-			StatType = STATTYPE_AccumulatorFLOAT;
-		}
-		if (StatType != STATTYPE_Error)
-		{
-			FindOrAddStat(LongName.NameAndInfo, StatType);
-		}
-
-		// Threads metadata.
-		const bool bIsThread = FStatConstants::NAME_ThreadGroup == GroupName;
-		if( bIsThread )
-		{
-			FindOrAddThread( LongName.NameAndInfo );
-		}	
-	}
+// 	// loop through the stats meta data messages
+// 	for (auto It = CurrentThreadState.ShortNameToLongName.CreateConstIterator(); It; ++It)
+// 	{
+// 		FStatMessage const& LongName = It.Value();
+// 		const FName GroupName = LongName.NameAndInfo.GetGroupName();
+// 
+// 		uint32 StatType = STATTYPE_Error;
+// 		if (LongName.NameAndInfo.GetField<EStatDataType>() == EStatDataType::ST_int64)
+// 		{
+// 			if( LongName.NameAndInfo.GetFlag( EStatMetaFlags::IsCycle ) )
+// 			{
+// 				StatType = STATTYPE_CycleCounter;
+// 			}
+// 			else if( LongName.NameAndInfo.GetFlag( EStatMetaFlags::IsMemory ) )
+// 			{
+// 				StatType = STATTYPE_MemoryCounter;
+// 			}
+// 			else
+// 			{
+// 				StatType = STATTYPE_AccumulatorDWORD;
+// 			}
+// 		}
+// 		else if (LongName.NameAndInfo.GetField<EStatDataType>() == EStatDataType::ST_double)
+// 		{
+// 			StatType = STATTYPE_AccumulatorFLOAT;
+// 		}
+// 		if (StatType != STATTYPE_Error)
+// 		{
+// 			FindOrAddStat(LongName.NameAndInfo, StatType);
+// 		}
+// 
+// 		// Threads metadata.
+// 		const bool bIsThread = FStatConstants::NAME_ThreadGroup == GroupName;
+// 		if( bIsThread )
+// 		{
+// 			FindOrAddThread( LongName.NameAndInfo );
+// 		}	
+// 	}
 }
 
 int32 FServiceConnection::FindOrAddStat( const FStatNameAndInfo& StatNameAndInfo, uint32 StatType)
@@ -939,14 +903,11 @@ int32 FServiceConnection::FindOrAddStat( const FStatNameAndInfo& StatNameAndInfo
 			GroupDescription.Name.RemoveFromStart(TEXT("STATGROUP_"));
 
 			// add to the meta data
-			FScopeLock ScopeLock(&CriticalSection);
-			MetaData.GroupDescriptions.Add(GroupDescription.ID, GroupDescription);
-		}
-		{
-			StatDescription.GroupID = GroupID;
-			FScopeLock ScopeLock(&CriticalSection);
-			MetaData.StatDescriptions.Add(StatDescription.ID, StatDescription);
+			StatMetaData.GroupDescriptions.Add(GroupDescription.ID, GroupDescription);
 		}
+
+		StatDescription.GroupID = GroupID;
+		StatMetaData.StatDescriptions.Add(StatDescription.ID, StatDescription);
 	}
 	// return the stat id
 	return StatID;
@@ -963,10 +924,9 @@ int32 FServiceConnection::FindOrAddThread(const FStatNameAndInfo& Thread)
 	const FName ShortName = Thread.GetShortName();
 
 	// add to the meta data
-	FScopeLock ScopeLock( &CriticalSection );
-	const int32 OldNum = MetaData.ThreadDescriptions.Num();
-	MetaData.ThreadDescriptions.Add( ThreadID, ShortName.ToString() );
-	const int32 NewNum = MetaData.ThreadDescriptions.Num();
+	const int32 OldNum = StatMetaData.ThreadDescriptions.Num();
+	StatMetaData.ThreadDescriptions.Add( ThreadID, ShortName.ToString() );
+	const int32 NewNum = StatMetaData.ThreadDescriptions.Num();
 
 	// meta data has been updated
 	CurrentData.MetaDataUpdated = CurrentData.MetaDataUpdated || OldNum != NewNum;
@@ -980,7 +940,6 @@ void FServiceConnection::GenerateAccumulators(TArray<FStatMessage>& Stats, TArra
 	for (int32 Index = 0; Index < Stats.Num(); ++Index)
 	{
 		const FStatMessage& StatMessage = Stats[Index];
-		const FName GroupName = StatMessage.NameAndInfo.GetGroupName();
 
 		uint32 StatType = STATTYPE_Error;
 		if (StatMessage.NameAndInfo.GetField<EStatDataType>() == EStatDataType::ST_int64)
@@ -1023,11 +982,10 @@ void FServiceConnection::GenerateAccumulators(TArray<FStatMessage>& Stats, TArra
 				Data.Value = StatMessage.GetValue_double();
 				FloatAccumulators.Add( Data );
 
-				const FName StatName = StatMessage.NameAndInfo.GetShortName();
-				if (StatName == FStatConstants::NAME_SecondsPerCycle)
+				const FName StatName = StatMessage.NameAndInfo.GetRawName();
+				if (StatName == FStatConstants::RAW_SecondsPerCycle)
 				{
-					FScopeLock ScopeLock( &CriticalSection );
-					MetaData.SecondsPerCycle = StatMessage.GetValue_double();
+					StatMetaData.SecondsPerCycle = StatMessage.GetValue_double();
 				}
 			}
 		}
@@ -1101,97 +1059,12 @@ void FServiceConnection::GenerateCycleGraphs(const FRawStatStackNode& Root, TMap
 	}
 }
 
-bool FServiceConnection::ReadAndConvertStatMessages( FArchive& Reader, bool bUseInAsync )
-{
-	// Buffer used to store the compressed and decompressed data.
-	TArray<uint8> SrcData;
-	TArray<uint8> DestData;
-
-	const bool bHasCompressedData = Stream.Header.HasCompressedData();
-	const bool bIsFinalized = Stream.Header.IsFinalized();
-
-	SCOPE_CYCLE_COUNTER( STAT_PC_ReadStatMessages );
-	if( bHasCompressedData )
-	{
-		while( Reader.Tell() < Reader.TotalSize() )
-		{
-			// Read the compressed data.
-			FCompressedStatsData UncompressedData( SrcData, DestData );
-			Reader << UncompressedData;
-			if( UncompressedData.HasReachedEndOfCompressedData() )
-			{
-				GenerateNewFrame( FStatMessage( FStatConstants::AdvanceFrame.GetEncodedName(), EStatOperation::AdvanceFrameEventGameThread, 0ll, false ), Reader, bUseInAsync );
-				return true;
-			}
-
-			FMemoryReader MemoryReader( DestData );
-
-			while( MemoryReader.Tell() < MemoryReader.TotalSize() )
-			{
-				// read the message
-				FStatMessage Message( Stream.ReadMessage( MemoryReader, bIsFinalized ) );
-				ReadMessages++;
-
-				if( Message.NameAndInfo.GetShortName() != TEXT( "Unknown FName" ) )
-				{
-					if( Message.NameAndInfo.GetField<EStatOperation>() == EStatOperation::AdvanceFrameEventGameThread && ReadMessages > 2 )
-					{
-						GenerateNewFrame( Message, Reader, bUseInAsync );
-					}
-
-					new (Messages)FStatMessage( Message );			
-				}
-				else
-				{
-					break;
-				}	
-			}
-
-			if( !bUseInAsync && DataFrames.Num() < FProfilerClientManager::MaxFramesPerTick )
-			{
-				return false;
-			}
-		}
-	}
-	// Obsolete. Remove later.
-	
-	if( Reader.Tell() >= Reader.TotalSize() )
-	{
-		return true;
-	}
-
-	return false;
-}
-
-void FServiceConnection::GenerateNewFrame( FStatMessage Message, FArchive &Reader, bool bUseInAsync )
-{
-	AddCollectedStatMessages( Message );
-
-	// create an old format data frame from the data
-	GenerateProfilerDataFrame();
-
-	{
-		// add the frame to the work list
-		FScopeLock ScopeLock( &CriticalSection );
-		DataFrames.Add( CurrentData );
-		DataLoadingProgress = (double)Reader.Tell() / (double)Reader.TotalSize();
-	}
-
-	if (bUseInAsync)
-	{
-		if (DataFrames.Num() > FProfilerClientManager::MaxFramesPerTick)
-		{
-			FPlatformProcess::Sleep( 0.1f );
-		}
-	}
-}
-
 void FServiceConnection::AddCollectedStatMessages( FStatMessage Message )
 {
 	SCOPE_CYCLE_COUNTER( STAT_PC_AddStatMessages );
-	new (Messages)FStatMessage( Message );
-	CurrentThreadState.AddMessages( Messages );
-	Messages.Reset();
+	new (PendingStatMessagesMessages)FStatMessage( Message );
+	CurrentThreadState.AddCondensedMessages( PendingStatMessagesMessages );
+	PendingStatMessagesMessages.Reset();
 }
 
 void FServiceConnection::GenerateProfilerDataFrame()
diff --git a/Engine/Source/Developer/ProfilerClient/Private/ProfilerClientManager.h b/Engine/Source/Developer/ProfilerClient/Private/ProfilerClientManager.h
--- a/Engine/Source/Developer/ProfilerClient/Private/ProfilerClientManager.h
+++ b/Engine/Source/Developer/ProfilerClient/Private/ProfilerClientManager.h
@@ -2,29 +2,17 @@
 
 #pragma once
 
-// #YRX_Profiler: 2015-05-13 Remove it later.
-#define PROFILER_THREADED_LOAD 1
-
+class FNewStatsReader;
+class FProfilerClientManager;
 
 /** Helper struct containing all of the data and operations associated with a service connection */
 struct FServiceConnection
 {
-	FServiceConnection()
-		: DataLoadingProgress( 0.0f )
-		, ReadMessages(0)
-	{
-		MetaData.CriticalSection = &CriticalSection;
-	}
+	FServiceConnection();
 
-	~FServiceConnection()
-	{
-		MetaData.CriticalSection = NULL;
-	}
+	~FServiceConnection();
 
-	FServiceConnection(const FServiceConnection& InConnnection)
-	{
-		MetaData.CriticalSection = &CriticalSection;
-	}
+	FServiceConnection(const FServiceConnection& InConnnection);
 
 	/** Instance Id */
 	FGuid InstanceId;
@@ -33,7 +21,7 @@ struct FServiceConnection
 	FMessageAddress ProfilerServiceAddress;
 
 	/** Descriptions for the stats */
-	FStatMetaData MetaData;
+	FStatMetaData StatMetaData;
 
 	/** Current frame worth of data */
 	FProfilerDataFrame CurrentData;
@@ -50,31 +38,22 @@ struct FServiceConnection
 	TMap<FName, int32> GroupNameArray;
 
 	/** Provides the long stat name to StatId mapping. */
-	TMap<FName,int32> LongNameToStatID;
+	TMap<FName, int32> LongNameToStatID;
 
 #if STATS
 	/** Stream reader */
 	FStatsReadStream Stream;
 
 	/** Pending stat messages */
-	TArray<FStatMessage> Messages;
+	TArray<FStatMessage> PendingStatMessagesMessages;
 #endif
 
 	/** Pending Data frames on a load */
 	TArray<FProfilerDataFrame> DataFrames;
 
-	/** Current data loading progress. */
-	float DataLoadingProgress;
-
-	/** Current number of read messages. */
-	uint64 ReadMessages;
-
 	/** Messages received and pending process, they are stored in a map as we can receive them out of order */
 	TMap<int64, TArray<uint8> > PendingMessages;
 
-	/** Critical Section needed to synchronize */
-	FCriticalSection CriticalSection;
-
 #if STATS
 	// generates the old style cycle graph
 	void GenerateCycleGraphs(const FRawStatStackNode& Root, TMap<uint32, FProfilerCycleGraph>& Graphs);
@@ -92,18 +71,7 @@ struct FServiceConnection
 	int32 FindOrAddThread(const FStatNameAndInfo& Thread);
 
 	// updates the meta date for the given instance
-	void UpdateMetaData();
-
-	
-	/**
-	 * Reads stat message from the archive and converts them to be usable by the profiler..
-	 *
-	 * @return true if read a special marker that indicates end of file
-	 */
-	bool ReadAndConvertStatMessages( FArchive& Reader, bool bUseInAsync );
-
-	/** Generates a new frame for the profiler. */
-	void GenerateNewFrame( FStatMessage Message, FArchive &Reader, bool bUseInAsync );
+	void UpdateMetaData_DISABLED();
 
 	/** Adds all collected stat messages to the current stats thread state. */
 	void AddCollectedStatMessages( FStatMessage Message );
@@ -111,6 +79,14 @@ struct FServiceConnection
 	/** Generates profiler data frame based on the collected stat messages. */
 	void GenerateProfilerDataFrame();
 #endif
+
+	/*-----------------------------------------------------------------------------
+		Transition into new system
+	-----------------------------------------------------------------------------*/
+
+	void LoadCapture( const FString& DataFilepath, FProfilerClientManager* ProfilerClientManager );
+
+	FNewStatsReader* StatsReader;
 };
 
 /**
@@ -119,25 +95,7 @@ struct FServiceConnection
 class FProfilerClientManager
 	: public IProfilerClient
 {
-	class FAsyncReadWorker : public FNonAbandonableTask
-	{
-	public:
-		FServiceConnection* LoadConnection;
-		FArchive* FileReader;
-
-		/** Constructor */
-		FAsyncReadWorker(FServiceConnection* InConnection, FArchive* InReader)
-			: LoadConnection(InConnection)
-			, FileReader(InReader)
-		{}
-
-		void DoWork();
-
-		FORCEINLINE TStatId GetStatId() const
-		{
-			RETURN_QUICK_DECLARE_CYCLE_STAT(FAsyncReadWorker, STATGROUP_ThreadPoolAsyncTasks);
-		}
-	};
+	friend class FNewStatsReader;
 
 public:
 
@@ -157,7 +115,6 @@ public:
 
 	virtual void Subscribe( const FGuid& Session ) override;
 	virtual void Track( const FGuid& Instance ) override;
-	virtual void Track( const TArray<TSharedPtr<ISessionInstanceInfo>>& Instances ) override;
 	virtual void Untrack( const FGuid& Instance ) override;
 	virtual void Unsubscribe() override;
 	virtual void SetCaptureState( const bool bRequestedCaptureState, const FGuid& InstanceId = FGuid() ) override;
@@ -167,7 +124,7 @@ public:
 
 	virtual const FStatMetaData& GetStatMetaData( const FGuid& InstanceId ) const override
 	{
-		return Connections.Find(InstanceId)->MetaData;
+		return Connections.Find(InstanceId)->StatMetaData;
 	}
 
 	virtual FProfilerClientDataDelegate& OnProfilerData() override
@@ -206,13 +163,6 @@ public:
 		return ProfilerLoadCompletedDelegate;
 	}
 
-	virtual FProfilerLoadedMetaDataDelegate& OnLoadedMetaData() override
-	{
-		return ProfilerLoadedMetaDataDelegate;
-	}
-
-	static const int32 MaxFramesPerTick = 30;
-
 private:
 
 	// Handles message bus shutdowns.
@@ -242,9 +192,6 @@ private:
 	/** If hash is ok, writes the data to the archive and returns true, otherwise only returns false. */
 	bool CheckHashAndWrite( const FProfilerServiceFileChunk& FileChunk, const FProfilerFileChunkHeader& FileChunkHeader, FArchive* Writer );
 
-	/** Broadcast that meta data has been updated. */
-	void BroadcastMetadataUpdate();
-
 	/** Broadcast that loading has completed and cleans internal structures. */
 	void FinalizeLoading();
 
@@ -252,7 +199,13 @@ private:
 	void DecompressDataAndSendToGame( FProfilerServiceData2* ToProcess );
 
 	/** Sends decompressed data to the game thread. */
-	void SendToGame( TArray<uint8>* DataToGame, int64 Frame, const FGuid InstanceId );
+	void SendDataToGame( TArray<uint8>* DataToGame, int64 Frame, const FGuid InstanceId );
+
+	/** A new profiler data frame from the async loading thread. */
+	void SendProfilerDataFrameToGame( FProfilerDataFrame* NewData, FStatMetaData* MetaDataPtr );
+
+	/** Removes active transfers and core tickers. */
+	void Shutdown();
 
 private:
 
@@ -264,6 +217,7 @@ private:
 	TArray<FGuid> PendingInstances;
 
 	/** Service connections */
+	// #YRX_tag 2015-11-19 should be only one active connection.
 	TMap<FGuid, FServiceConnection> Connections;
 
 	struct FReceivedFileInfo
@@ -325,9 +279,6 @@ private:
 	/** Delegate for notifying clients of a load completion */
 	FProfilerLoadCompletedDelegate ProfilerLoadCompletedDelegate;
 
-	/** Delegate for notifying clients of initial meta data load */
-	FProfilerLoadedMetaDataDelegate ProfilerLoadedMetaDataDelegate;
-
 	/** Holds a delegate to be invoked when the widget ticks. */
 	FTickerDelegate TickDelegate;
 
@@ -349,18 +300,6 @@ private:
 	/** Handle to the registered OnShutdown for Message Bus. */
 	FDelegateHandle OnShutdownMessageBusDelegateHandle;
 
-#if PROFILER_THREADED_LOAD
-	/** Loads a file asynchronously */
-	bool AsyncLoad();
-
-	/** Load Task */
-	FAsyncTask<FAsyncReadWorker>* LoadTask;
-#else
-	/** Loads a file synchronously */
-	bool SyncLoad();
-	/** File reader */
-	FArchive* FileReader;
-#endif
 
 	/** Holds the last time a ping was made to instances */
 	FDateTime LastPingTime;
diff --git a/Engine/Source/Developer/ProfilerClient/Public/Interfaces/IProfilerClient.h b/Engine/Source/Developer/ProfilerClient/Public/Interfaces/IProfilerClient.h
--- a/Engine/Source/Developer/ProfilerClient/Public/Interfaces/IProfilerClient.h
+++ b/Engine/Source/Developer/ProfilerClient/Public/Interfaces/IProfilerClient.h
@@ -8,7 +8,7 @@ class ISessionInstanceInfo;
 
 
 /** Delegate for passing profiler data to UI */
-DECLARE_MULTICAST_DELEGATE_ThreeParams(FProfilerClientDataDelegate, const FGuid&, const FProfilerDataFrame&, const float );
+DECLARE_MULTICAST_DELEGATE_TwoParams(FProfilerClientDataDelegate, const FGuid&, const FProfilerDataFrame& );
 
 /** Delegate for alerting UI a session has been established */
 DECLARE_MULTICAST_DELEGATE_TwoParams(FProfilerClientConnectedDelegate, const FGuid&, const FGuid& );
@@ -17,7 +17,7 @@ DECLARE_MULTICAST_DELEGATE_TwoParams(FProfilerClientConnectedDelegate, const FGu
 DECLARE_MULTICAST_DELEGATE_TwoParams(FProfilerClientDisconnectedDelegate, const FGuid&, const FGuid& );
 
 /** Delegate for alerting subscribers the meta data has been updated */
-DECLARE_MULTICAST_DELEGATE_OneParam(FProfilerMetaDataUpdateDelegate, const FGuid& );
+DECLARE_MULTICAST_DELEGATE_TwoParams(FProfilerMetaDataUpdateDelegate, const FGuid&, const FStatMetaData& );
 
 /** Delegate for alerting clients a load has started */
 DECLARE_MULTICAST_DELEGATE_OneParam(FProfilerLoadStartedDelegate, const FGuid& );
@@ -51,8 +51,6 @@ public:
 	 */
 	virtual void Track( const FGuid& Instance ) = 0;
 
-	virtual void Track( const TArray<TSharedPtr<ISessionInstanceInfo>>& Instances ) = 0;
-
 	/**
 	 * 
 	 */
@@ -148,13 +146,6 @@ public:
 	 */
 	virtual FProfilerLoadCompletedDelegate& OnLoadCompleted() = 0;
 
-	/**
-	 * Retrieves the meta data loaded delegate.
-	 *
-	 * @return profiler meta dataloaded delegate.
-	 */
-	virtual FProfilerLoadedMetaDataDelegate& OnLoadedMetaData() = 0;
-
 protected:
 
 	/** Virtual destructor. */
diff --git a/Engine/Source/Developer/ProfilerMessages/Classes/ProfilerServiceMessages.h b/Engine/Source/Developer/ProfilerMessages/Classes/ProfilerServiceMessages.h
--- a/Engine/Source/Developer/ProfilerMessages/Classes/ProfilerServiceMessages.h
+++ b/Engine/Source/Developer/ProfilerMessages/Classes/ProfilerServiceMessages.h
@@ -5,56 +5,51 @@
 #include "ProfilerServiceMessages.generated.h"
 
 
-/**
- */
+/** Profiler Service authorization message. */
 USTRUCT()
 struct FProfilerServiceAuthorize
 {
 	GENERATED_USTRUCT_BODY()
 
-	/** */
+	/** Session ID. */
 	UPROPERTY(EditAnywhere, Category="Message")
 	FGuid SessionId;
 
-	/** */
+	/** Instance ID. */
 	UPROPERTY(EditAnywhere, Category="Message")
 	FGuid InstanceId;
 
-	/** */
-	UPROPERTY(EditAnywhere, Category="Message")
-	TArray<uint8> Data;
-
 	/** Default constructor. */
 	FProfilerServiceAuthorize() { }
 
 	/** Creates and initializes a new instance. */
-	FProfilerServiceAuthorize(const FGuid& InSessionId, const FGuid& InInstanceId, const TArray<uint8>& InData)
+	FProfilerServiceAuthorize(const FGuid& InSessionId, const FGuid& InInstanceId)
 		: SessionId(InSessionId)
 		, InstanceId(InInstanceId)
-	{
-		Data.Append(InData);
-	}
+	{}
 
 };
 
 
-/**
- */
+/** Profiler Service data. */
 USTRUCT()
 struct FProfilerServiceData2
 {
 	GENERATED_USTRUCT_BODY()
 
-	/** */
+	/** Instance ID. */
 	UPROPERTY(EditAnywhere, Category="Message")
 	FGuid InstanceId;
 
+	/** Stats frame. */
 	UPROPERTY(EditAnywhere, Category="Message")
 	int64 Frame;
 
+	/** Size of the compressed data. */
 	UPROPERTY(EditAnywhere, Category="Message")
 	int32 CompressedSize;
 	
+	/** Size of the uncompressed data. */
 	UPROPERTY(EditAnywhere, Category="Message")
 	int32 UncompressedSize;
 
@@ -71,7 +66,7 @@ struct FProfilerServiceData2
 		, Frame(InFrame)
 		, CompressedSize(InCompressedSize)
 		, UncompressedSize(InUncompressedSize)
-		, HexData(MoveTemp(InHexData))
+		, HexData(InHexData)
 	{ }
 };
 
diff --git a/Engine/Source/Developer/ProfilerService/Private/ProfilerServiceFileTransfer.cpp b/Engine/Source/Developer/ProfilerService/Private/ProfilerServiceFileTransfer.cpp
new file mode 100644
--- /dev/null
+++ b/Engine/Source/Developer/ProfilerService/Private/ProfilerServiceFileTransfer.cpp
@@ -0,0 +1,247 @@
+// Copyright 1998-2015 Epic Games, Inc. All Rights Reserved.
+
+#include "ProfilerServicePrivatePCH.h"
+#include "SecureHash.h"
+
+
+FFileTransferRunnable::FFileTransferRunnable( FMessageEndpointPtr& InMessageEndpoint )
+	: Runnable( nullptr )
+	, WorkEvent( FPlatformProcess::GetSynchEventFromPool( true ) )
+	, MessageEndpoint( InMessageEndpoint )
+	, StopTaskCounter( 0 )
+{
+	Runnable = FRunnableThread::Create( this, TEXT( "FFileTransferRunnable" ), 128 * 1024, TPri_BelowNormal );
+}
+
+FFileTransferRunnable::~FFileTransferRunnable()
+{
+	if (Runnable != nullptr)
+	{
+		Stop();
+		Runnable->WaitForCompletion();
+		delete Runnable;
+		Runnable = nullptr;
+	}
+
+	// Delete all active file readers.
+	for (auto It = ActiveTransfers.CreateIterator(); It; ++It)
+	{
+		FReaderAndAddress& ReaderAndAddress = It.Value();
+		DeleteFileReader( ReaderAndAddress );
+
+		UE_LOG( LogProfilerService, Log, TEXT( "File service-client sending aborted (srv): %s" ), *It.Key() );
+	}
+
+	FPlatformProcess::ReturnSynchEventToPool( WorkEvent );
+	WorkEvent = nullptr;
+}
+
+bool FFileTransferRunnable::Init()
+{
+	return true;
+}
+
+uint32 FFileTransferRunnable::Run()
+{
+	while (!ShouldStop())
+	{
+		if (WorkEvent->Wait( 250 ))
+		{
+			FProfilerServiceFileChunk* FileChunk;
+			while (!ShouldStop() && SendQueue.Dequeue( FileChunk ))
+			{
+				FMemoryReader MemoryReader( FileChunk->Header );
+				FProfilerFileChunkHeader FileChunkHeader;
+				MemoryReader << FileChunkHeader;
+				FileChunkHeader.Validate();
+
+				if (FileChunkHeader.ChunkType == EProfilerFileChunkType::SendChunk)
+				{
+					// Find the corresponding file archive reader and recipient.
+					FArchive* ReaderArchive = nullptr;
+					FMessageAddress Recipient;
+					{
+						FScopeLock Lock( &SyncActiveTransfers );
+						const FReaderAndAddress* ReaderAndAddress = ActiveTransfers.Find( FileChunk->Filename );
+						if (ReaderAndAddress)
+						{
+							ReaderArchive = ReaderAndAddress->Key;
+							Recipient = ReaderAndAddress->Value;
+						}
+					}
+
+					// If there is no reader and recipient is invalid, it means that the file transfer is no longer valid, because client disconnected or exited.
+					if (ReaderArchive && Recipient.IsValid())
+					{
+						ReadAndSetHash( FileChunk, FileChunkHeader, ReaderArchive );
+
+						if (MessageEndpoint.IsValid())
+						{
+							MessageEndpoint->Send( FileChunk, Recipient );
+						}
+					}
+				}
+				else if (FileChunkHeader.ChunkType == EProfilerFileChunkType::PrepareFile)
+				{
+					PrepareFileForSending( FileChunk );
+				}
+			}
+
+			WorkEvent->Reset();
+		}
+	}
+
+	return 0;
+}
+
+void FFileTransferRunnable::Exit()
+{
+
+}
+
+void FFileTransferRunnable::EnqueueFileToSend( const FString& StatFilename, const FMessageAddress& RecipientAddress, const FGuid& ServiceInstanceId )
+{
+	UE_LOG( LogProfilerService, Log, TEXT( "Opening stats file for service-client sending: %s" ), *StatFilename );
+
+	const int64 FileSize = IFileManager::Get().FileSize( *StatFilename );
+	if (FileSize < 4)
+	{
+		UE_LOG( LogProfilerService, Error, TEXT( "Could not open: %s" ), *StatFilename );
+		return;
+	}
+
+	FArchive* FileReader = IFileManager::Get().CreateFileReader( *StatFilename );
+	if (!FileReader)
+	{
+		UE_LOG( LogProfilerService, Error, TEXT( "Could not open: %s" ), *StatFilename );
+		return;
+	}
+
+	{
+		FScopeLock Lock( &SyncActiveTransfers );
+		check( !ActiveTransfers.Contains( StatFilename ) );
+		ActiveTransfers.Add( StatFilename, FReaderAndAddress( FileReader, RecipientAddress ) );
+	}
+
+	// This is not a real file chunk, but helper to prepare file for sending on the runnable.
+	EnqueueFileChunkToSend( new FProfilerServiceFileChunk( ServiceInstanceId, StatFilename, FProfilerFileChunkHeader( 0, 0, FileReader->TotalSize(), EProfilerFileChunkType::PrepareFile ).AsArray() ), true );
+}
+
+void FFileTransferRunnable::EnqueueFileChunkToSend( FProfilerServiceFileChunk* FileChunk, bool bTriggerWorkEvent /*= false */ )
+{
+	SendQueue.Enqueue( FileChunk );
+
+	if (bTriggerWorkEvent)
+	{
+		// Trigger the runnable.
+		WorkEvent->Trigger();
+	}
+}
+
+void FFileTransferRunnable::PrepareFileForSending( FProfilerServiceFileChunk*& FileChunk )
+{
+	// Find the corresponding file archive and recipient.
+	FArchive* Reader = nullptr;
+	FMessageAddress Recipient;
+	{
+		FScopeLock Lock( &SyncActiveTransfers );
+		const FReaderAndAddress& ReaderAndAddress = ActiveTransfers.FindChecked( FileChunk->Filename );
+		Reader = ReaderAndAddress.Key;
+		Recipient = ReaderAndAddress.Value;
+	}
+
+	int64 ChunkOffset = 0;
+	int64 RemainingSizeToSend = Reader->TotalSize();
+
+	while (RemainingSizeToSend > 0)
+	{
+		const int64 SizeToCopy = FMath::Min( FProfilerFileChunkHeader::DefChunkSize, RemainingSizeToSend );
+
+		EnqueueFileChunkToSend( new FProfilerServiceFileChunk( FileChunk->InstanceId, FileChunk->Filename, FProfilerFileChunkHeader( ChunkOffset, SizeToCopy, Reader->TotalSize(), EProfilerFileChunkType::SendChunk ).AsArray() ) );
+
+		ChunkOffset += SizeToCopy;
+		RemainingSizeToSend -= SizeToCopy;
+	}
+
+	// Trigger the runnable.
+	WorkEvent->Trigger();
+
+	// Delete this file chunk.
+	delete FileChunk;
+	FileChunk = nullptr;
+}
+
+void FFileTransferRunnable::FinalizeFileSending( const FString& Filename )
+{
+	FScopeLock Lock( &SyncActiveTransfers );
+
+	check( ActiveTransfers.Contains( Filename ) );
+	FReaderAndAddress ReaderAndAddress = ActiveTransfers.FindAndRemoveChecked( Filename );
+	DeleteFileReader( ReaderAndAddress );
+
+	UE_LOG( LogProfilerService, Log, TEXT( "File service-client sent successfully : %s" ), *Filename );
+}
+
+void FFileTransferRunnable::AbortFileSending( const FMessageAddress& Recipient )
+{
+	FScopeLock Lock( &SyncActiveTransfers );
+
+	for (auto It = ActiveTransfers.CreateIterator(); It; ++It)
+	{
+		FReaderAndAddress& ReaderAndAddress = It.Value();
+		if (ReaderAndAddress.Value == Recipient)
+		{
+			UE_LOG( LogProfilerService, Log, TEXT( "File service-client sending aborted (cl): %s" ), *It.Key() );
+			FReaderAndAddress ActiveReaderAndAddress = ActiveTransfers.FindAndRemoveChecked( It.Key() );
+			DeleteFileReader( ActiveReaderAndAddress );
+		}
+	}
+}
+
+void FFileTransferRunnable::DeleteFileReader( FReaderAndAddress& ReaderAndAddress )
+{
+	delete ReaderAndAddress.Key;
+	ReaderAndAddress.Key = nullptr;
+}
+
+void FFileTransferRunnable::ReadAndSetHash( FProfilerServiceFileChunk* FileChunk, const FProfilerFileChunkHeader& FileChunkHeader, FArchive* Reader )
+{
+	TArray<uint8> FileChunkData;
+
+	FileChunkData.AddUninitialized( FileChunkHeader.ChunkSize );
+
+	Reader->Seek( FileChunkHeader.ChunkOffset );
+	Reader->Serialize( FileChunkData.GetData(), FileChunkHeader.ChunkSize );
+
+	const int32 HashSize = 20;
+	uint8 LocalHash[HashSize] = {0};
+
+	// Hash file chunk data. 
+	FSHA1 Sha;
+	Sha.Update( FileChunkData.GetData(), FileChunkHeader.ChunkSize );
+	// Hash file chunk header.
+	Sha.Update( FileChunk->Header.GetData(), FileChunk->Header.Num() );
+	Sha.Final();
+	Sha.GetHash( LocalHash );
+
+	FileChunk->ChunkHash.AddUninitialized( HashSize );
+	FMemory::Memcpy( FileChunk->ChunkHash.GetData(), LocalHash, HashSize );
+
+	// Limit transfer per second, otherwise we will probably hang the message bus.
+	static int64 TotalReadBytes = 0;
+#if	_DEBUG
+	static const int64 NumBytesPerTick = 128 * 1024;
+#else
+	static const int64 NumBytesPerTick = 256 * 1024;
+#endif // _DEBUG
+
+	// Convert to hex.
+	FileChunk->HexData = FString::FromHexBlob( FileChunkData.GetData(), FileChunkData.Num() );
+
+	TotalReadBytes += FileChunkHeader.ChunkSize;
+	if (TotalReadBytes > NumBytesPerTick)
+	{
+		FPlatformProcess::Sleep( 0.1f );
+		TotalReadBytes = 0;
+	}
+}
diff --git a/Engine/Source/Developer/ProfilerService/Private/ProfilerServiceManager.cpp b/Engine/Source/Developer/ProfilerService/Private/ProfilerServiceManager.cpp
--- a/Engine/Source/Developer/ProfilerService/Private/ProfilerServiceManager.cpp
+++ b/Engine/Source/Developer/ProfilerService/Private/ProfilerServiceManager.cpp
@@ -4,304 +4,44 @@
 
 #include "StatsData.h"
 #include "StatsFile.h"
-#include "SecureHash.h"
 
 
-DEFINE_LOG_CATEGORY_STATIC( LogProfilerService, Log, All );
+DEFINE_LOG_CATEGORY( LogProfilerService );
 
-/** 
- * Thread used to read, prepare and send files through the message bus.
- * Supports resending bad file chunks and basic synchronization between service and client.
- */
-class FFileTransferRunnable : public FRunnable
-{
-	typedef TKeyValuePair<FArchive*,FMessageAddress> FReaderAndAddress;
-
-public:	
-
-	/** Default constructor. */
-	FFileTransferRunnable( FMessageEndpointPtr& InMessageEndpoint )
-		: Runnable( nullptr )
-		, WorkEvent( FPlatformProcess::GetSynchEventFromPool( true ) )
-		, MessageEndpoint( InMessageEndpoint )
-		, StopTaskCounter( 0 )
-	{
-		Runnable = FRunnableThread::Create(this, TEXT("FFileTransferRunnable"), 128 * 1024, TPri_BelowNormal);
-	}
-
-	/** Destructor. */
-	~FFileTransferRunnable()
-	{
-		if( Runnable != nullptr )
-		{
-			Stop();
-			Runnable->WaitForCompletion();
-			delete Runnable;
-			Runnable = nullptr;
-		}
-
-		// Delete all active file readers.
-		for( auto It = ActiveTransfers.CreateIterator(); It; ++It )
-		{
-			FReaderAndAddress& ReaderAndAddress = It.Value();
-			DeleteFileReader( ReaderAndAddress );
-
-			UE_LOG( LogProfilerService, Log, TEXT( "File service-client sending aborted (srv): %s" ), *It.Key() );
-		}
-
-		FPlatformProcess::ReturnSynchEventToPool( WorkEvent );
-		WorkEvent = nullptr;
-	}
-
-	// Begin FRunnable interface.
-	virtual bool Init()
-	{
-		return true;
-	}
-
-	virtual uint32 Run()
-	{
-		while( !ShouldStop() )
-		{
-			if( WorkEvent->Wait( 250 ) )
-			{
-				FProfilerServiceFileChunk* FileChunk;
-				while( !ShouldStop() && SendQueue.Dequeue( FileChunk ) )
-				{
-					FMemoryReader MemoryReader(FileChunk->Header);
-					FProfilerFileChunkHeader FileChunkHeader;
-					MemoryReader << FileChunkHeader;
-					FileChunkHeader.Validate();
-
-					if( FileChunkHeader.ChunkType == EProfilerFileChunkType::SendChunk )
-					{
-						// Find the corresponding file archive reader and recipient.
-						FArchive* ReaderArchive = nullptr;
-						FMessageAddress Recipient;
-						{
-							FScopeLock Lock(&SyncActiveTransfers);
-							const FReaderAndAddress* ReaderAndAddress = ActiveTransfers.Find( FileChunk->Filename );
-							if( ReaderAndAddress )
-							{
-								ReaderArchive = ReaderAndAddress->Key;
-								Recipient = ReaderAndAddress->Value;
-							}
-						}
-
-						// If there is no reader and recipient is invalid, it means that the file transfer is no longer valid, because client disconnected or exited.
-						if( ReaderArchive && Recipient.IsValid() )
-						{
-							ReadAndSetHash( FileChunk, FileChunkHeader, ReaderArchive );
-
-							if( MessageEndpoint.IsValid() )
-							{
-								MessageEndpoint->Send( FileChunk, Recipient );
-							}
-						}		
-					}
-					else if( FileChunkHeader.ChunkType == EProfilerFileChunkType::PrepareFile )
-					{
-						PrepareFileForSending( FileChunk );
-					}
-				}
-
-				WorkEvent->Reset();
-			}
-		}
-
-		return 0;
-	}
-
-	virtual void Stop()
-	{
-		StopTaskCounter.Increment();
-	}
-
-	virtual void Exit()
-	{}
-	// End FRunnable interface
-
-	void EnqueueFileToSend( const FString& StatFilename, const FMessageAddress& RecipientAddress, const FGuid& ServiceInstanceId )
-	{
-		UE_LOG( LogProfilerService, Log, TEXT( "Opening stats file for service-client sending: %s" ), *StatFilename );
-
-		const int64 FileSize = IFileManager::Get().FileSize(*StatFilename);
-		if( FileSize < 4 )
-		{
-			UE_LOG( LogProfilerService, Error, TEXT( "Could not open: %s" ), *StatFilename );
-			return;
-		}
-
-		FArchive* FileReader = IFileManager::Get().CreateFileReader(*StatFilename);
-		if( !FileReader )
-		{
-			UE_LOG( LogProfilerService, Error, TEXT( "Could not open: %s" ), *StatFilename );
-			return;
-		}
-
-		{
-			FScopeLock Lock(&SyncActiveTransfers);
-			check( !ActiveTransfers.Contains( StatFilename ) );
-			ActiveTransfers.Add( StatFilename, FReaderAndAddress(FileReader,RecipientAddress) );
-		}
-
-		// This is not a real file chunk, but helper to prepare file for sending on the runnable.
-		EnqueueFileChunkToSend( new FProfilerServiceFileChunk( ServiceInstanceId,StatFilename,FProfilerFileChunkHeader(0,0,FileReader->TotalSize(),EProfilerFileChunkType::PrepareFile).AsArray() ), true );
-	}
-
-	/** Enqueues a file chunk. */
-	void EnqueueFileChunkToSend( FProfilerServiceFileChunk* FileChunk, bool bTriggerWorkEvent = false )
-	{
-		SendQueue.Enqueue( FileChunk );
+/*-----------------------------------------------------------------------------
+	FProfilerServiceStatsStream
+-----------------------------------------------------------------------------*/
 
-		if( bTriggerWorkEvent )
-		{
-			// Trigger the runnable.
-			WorkEvent->Trigger();
-		}
-	}
-
-	/** Prepare the chunks to be sent through the message bus. */
-	void PrepareFileForSending( FProfilerServiceFileChunk*& FileChunk )
-	{
-		// Find the corresponding file archive and recipient.
-		FArchive* Reader = nullptr;
-		FMessageAddress Recipient;
-		{
-			FScopeLock Lock(&SyncActiveTransfers);
-			const FReaderAndAddress& ReaderAndAddress = ActiveTransfers.FindChecked( FileChunk->Filename );
-			Reader = ReaderAndAddress.Key;
-			Recipient = ReaderAndAddress.Value;
-		}
-
-		int64 ChunkOffset = 0;
-		int64 RemainingSizeToSend = Reader->TotalSize();
-
-		while( RemainingSizeToSend > 0 )
-		{
-			const int64 SizeToCopy = FMath::Min( FProfilerFileChunkHeader::DefChunkSize, RemainingSizeToSend );
-
-			EnqueueFileChunkToSend( new FProfilerServiceFileChunk( FileChunk->InstanceId,FileChunk->Filename,FProfilerFileChunkHeader(ChunkOffset,SizeToCopy,Reader->TotalSize(),EProfilerFileChunkType::SendChunk).AsArray() ) );
-
-			ChunkOffset += SizeToCopy;
-			RemainingSizeToSend -= SizeToCopy;
-		}
-
-		// Trigger the runnable.
-		WorkEvent->Trigger();
-
-		// Delete this file chunk.
-		delete FileChunk;
-		FileChunk = nullptr;	
-	}
-
-	/** Removes file from the list of the active transfers, must be confirmed by the profiler client. */
-	void FinalizeFileSending( const FString& Filename )
-	{
-		FScopeLock Lock(&SyncActiveTransfers);
-
-		check( ActiveTransfers.Contains( Filename ) );
-		FReaderAndAddress ReaderAndAddress = ActiveTransfers.FindAndRemoveChecked( Filename );
-		DeleteFileReader( ReaderAndAddress );
-
-		UE_LOG( LogProfilerService, Log, TEXT( "File service-client sent successfully : %s" ), *Filename );
-	}
+#if STATS
 
-	/** Aborts file sending to the specified client, probably client disconnected or exited. */
-	void AbortFileSending( const FMessageAddress& Recipient )
+/** Helper class for writing the condensed messages with related metadata. */
+class FProfilerServiceStatsStream : protected FStatsWriteStream
+{
+public:
+	void WriteFrameMessagesWithMetadata( int64 TargetFrame, bool bNeedFullMetadata )
 	{
-		FScopeLock Lock(&SyncActiveTransfers);
+		FMemoryWriter Ar( OutData, false, true );
 
-		for( auto It = ActiveTransfers.CreateIterator(); It; ++It )
+		if (bNeedFullMetadata)
 		{
-			FReaderAndAddress& ReaderAndAddress = It.Value();
-			if( ReaderAndAddress.Value == Recipient )
-			{
-				UE_LOG( LogProfilerService, Log, TEXT( "File service-client sending aborted (cl): %s" ), *It.Key() );
-				FReaderAndAddress ActiveReaderAndAddress = ActiveTransfers.FindAndRemoveChecked( It.Key() );
-				DeleteFileReader( ActiveReaderAndAddress );
-			}
+			WriteMetadata( Ar );
 		}
-	}
-
-	/** Checks if there has been any stop requests. */
-	FORCEINLINE bool ShouldStop() const
-	{
-		return StopTaskCounter.GetValue() > 0;
-	}
 
-protected:
-	/** Deletes the file reader. */
-	void DeleteFileReader( FReaderAndAddress& ReaderAndAddress )
-	{
-		delete ReaderAndAddress.Key;
-		ReaderAndAddress.Key = nullptr;
+		WriteCondensedMessages( Ar, TargetFrame );
 	}
 
-	/** Reads the data from the archive and generates hash. */
-	void ReadAndSetHash( FProfilerServiceFileChunk* FileChunk, const FProfilerFileChunkHeader& FileChunkHeader, FArchive* Reader )
+	/** Non-const for MoveTemp. */
+	TArray<uint8>& GetOutData()
 	{
-		TArray<uint8> FileChunkData;
-
-		FileChunkData.AddUninitialized( FileChunkHeader.ChunkSize );
-
-		Reader->Seek( FileChunkHeader.ChunkOffset );
-		Reader->Serialize( FileChunkData.GetData(), FileChunkHeader.ChunkSize );
-
-		const int32 HashSize = 20;
-		uint8 LocalHash[HashSize]={0};
-
-		// Hash file chunk data. 
-		FSHA1 Sha;
-		Sha.Update( FileChunkData.GetData(), FileChunkHeader.ChunkSize );
-		// Hash file chunk header.
-		Sha.Update( FileChunk->Header.GetData(), FileChunk->Header.Num() );
-		Sha.Final();
-		Sha.GetHash( LocalHash );
-
-		FileChunk->ChunkHash.AddUninitialized( HashSize );
-		FMemory::Memcpy( FileChunk->ChunkHash.GetData(), LocalHash, HashSize );
-
-		// Limit transfer per second, otherwise we will probably hang the message bus.
-		static int64 TotalReadBytes = 0;
-#if	_DEBUG
-		static const int64 NumBytesPerTick = 128*1024;
-#else
-		static const int64 NumBytesPerTick = 256*1024;
-#endif // _DEBUG
-
-		// Convert to hex.
-		FileChunk->HexData = FString::FromHexBlob( FileChunkData.GetData(), FileChunkData.Num() );
-
-		TotalReadBytes += FileChunkHeader.ChunkSize;
-		if( TotalReadBytes > NumBytesPerTick )
-		{
-			FPlatformProcess::Sleep( 0.1f );
-			TotalReadBytes = 0;
-		}
+		return OutData;
 	}
+};
 
-	/** Thread that is running this task. */
-	FRunnableThread* Runnable;
-
-	/** Event used to signaling that work is available. */
-	FEvent* WorkEvent;
-
-	/** Holds the messaging endpoint. */
-	FMessageEndpointPtr& MessageEndpoint;
-
-	/** > 0 if we have been asked to abort work in progress at the next opportunity. */
-	FThreadSafeCounter StopTaskCounter;
-
-	/** Added on the main thread, processed on the async thread. */
-	TQueue<FProfilerServiceFileChunk*,EQueueMode::Mpsc> SendQueue;
-
-	/** Critical section used to synchronize. */
-	FCriticalSection SyncActiveTransfers;
+#endif
 
-	/** Active transfers, stored as a filename -> reader and destination address. Assumes that filename is unique and never will be the same. */
-	TMap<FString,FReaderAndAddress> ActiveTransfers;
-};
+/*-----------------------------------------------------------------------------
+	FProfilerServiceManager
+-----------------------------------------------------------------------------*/
 
 void FProfilerServiceManager::HandleServiceFileChunkMessage( const FProfilerServiceFileChunk& Message, const IMessageContextRef& Context )
 {
@@ -581,14 +321,12 @@ void FProfilerServiceManager::HandleServiceSubscribeMessage( const FProfilerServ
 		FClientData Data;
 		Data.Active = true;
 		Data.Preview = false;
-		//Data.StatsWriteFile.WriteHeader();
 
 		// Add to the client list.
 		ClientData.Add( SenderAddress, Data );
 
-		// Send authorized and stat descriptions.
-		const TArray<uint8> OutData;// = ClientData.Find( SenderAddress )->StatsWriteFile.GetOutData();
-		MessageEndpoint->Send( new FProfilerServiceAuthorize( SessionId, InstanceId, OutData ), SenderAddress );
+		// Send authorize.
+		MessageEndpoint->Send( new FProfilerServiceAuthorize( SessionId, InstanceId ), SenderAddress );
 
 		// Initiate the ping callback
 		if (ClientData.Num() == 1)
@@ -645,11 +383,11 @@ void FProfilerServiceManager::HandleNewFrame(int64 Frame)
 	}
 
 	// Write frame.
-	FStatsWriteFile StatsWriteFile;
-	StatsWriteFile.WriteFrame( Frame, bNeedFullMetadata );
+	FProfilerServiceStatsStream StatsStream;
+	StatsStream.WriteFrameMessagesWithMetadata( Frame, bNeedFullMetadata );
 
 	// Task graph
-	TArray<uint8>* DataToTask = new TArray<uint8>( MoveTemp( StatsWriteFile.GetOutData() ) );
+	TArray<uint8>* DataToTask = new TArray<uint8>( MoveTemp( StatsStream.GetOutData() ) );
 
 	// Compression and encoding is done on the task graph
 	FSimpleDelegateGraphTask::CreateAndDispatchWhenReady
diff --git a/Engine/Source/Developer/ProfilerService/Private/ProfilerServiceManager.h b/Engine/Source/Developer/ProfilerService/Private/ProfilerServiceManager.h
--- a/Engine/Source/Developer/ProfilerService/Private/ProfilerServiceManager.h
+++ b/Engine/Source/Developer/ProfilerService/Private/ProfilerServiceManager.h
@@ -3,6 +3,86 @@
 #pragma once
 
 
+DECLARE_LOG_CATEGORY_EXTERN( LogProfilerService, Log, All );
+
+/**
+* Thread used to read, prepare and send files through the message bus.
+* Supports resending bad file chunks and basic synchronization between service and client.
+*/
+class FFileTransferRunnable : public FRunnable
+{
+	typedef TKeyValuePair<FArchive*, FMessageAddress> FReaderAndAddress;
+
+public:
+
+	/** Default constructor. */
+	FFileTransferRunnable( FMessageEndpointPtr& InMessageEndpoint );
+
+	/** Destructor. */
+	~FFileTransferRunnable();
+
+	// Begin FRunnable interface.
+	virtual bool Init();
+
+	virtual uint32 Run();
+
+	virtual void Stop()
+	{
+		StopTaskCounter.Increment();
+	}
+
+	virtual void Exit();
+	// End FRunnable interface
+
+	void EnqueueFileToSend( const FString& StatFilename, const FMessageAddress& RecipientAddress, const FGuid& ServiceInstanceId );
+
+	/** Enqueues a file chunk. */
+	void EnqueueFileChunkToSend( FProfilerServiceFileChunk* FileChunk, bool bTriggerWorkEvent = false );
+
+	/** Prepare the chunks to be sent through the message bus. */
+	void PrepareFileForSending( FProfilerServiceFileChunk*& FileChunk );
+
+	/** Removes file from the list of the active transfers, must be confirmed by the profiler client. */
+	void FinalizeFileSending( const FString& Filename );
+
+	/** Aborts file sending to the specified client, probably client disconnected or exited. */
+	void AbortFileSending( const FMessageAddress& Recipient );
+
+	/** Checks if there has been any stop requests. */
+	FORCEINLINE bool ShouldStop() const
+	{
+		return StopTaskCounter.GetValue() > 0;
+	}
+
+protected:
+	/** Deletes the file reader. */
+	void DeleteFileReader( FReaderAndAddress& ReaderAndAddress );
+
+	/** Reads the data from the archive and generates hash. */
+	void ReadAndSetHash( FProfilerServiceFileChunk* FileChunk, const FProfilerFileChunkHeader& FileChunkHeader, FArchive* Reader );
+
+	/** Thread that is running this task. */
+	FRunnableThread* Runnable;
+
+	/** Event used to signaling that work is available. */
+	FEvent* WorkEvent;
+
+	/** Holds the messaging endpoint. */
+	FMessageEndpointPtr& MessageEndpoint;
+
+	/** > 0 if we have been asked to abort work in progress at the next opportunity. */
+	FThreadSafeCounter StopTaskCounter;
+
+	/** Added on the main thread, processed on the async thread. */
+	TQueue<FProfilerServiceFileChunk*, EQueueMode::Mpsc> SendQueue;
+
+	/** Critical section used to synchronize. */
+	FCriticalSection SyncActiveTransfers;
+
+	/** Active transfers, stored as a filename -> reader and destination address. Assumes that filename is unique and never will be the same. */
+	TMap<FString, FReaderAndAddress> ActiveTransfers;
+};
+
 #if STATS
 
 /**
@@ -125,7 +205,7 @@ private:
 #endif // STATS
 
 	/** Thread used to read, prepare and send file chunks through the message bus. */
-	class FFileTransferRunnable* FileTransferRunnable;
+	FFileTransferRunnable* FileTransferRunnable;
 
 	/** Filename of last capture file. */
 	FString LastStatsFilename;
diff --git a/Engine/Source/Developer/ProfilerService/Public/Interfaces/IProfilerServiceManager.h b/Engine/Source/Developer/ProfilerService/Public/Interfaces/IProfilerServiceManager.h
--- a/Engine/Source/Developer/ProfilerService/Public/Interfaces/IProfilerServiceManager.h
+++ b/Engine/Source/Developer/ProfilerService/Public/Interfaces/IProfilerServiceManager.h
@@ -107,12 +107,13 @@ FORCEINLINE FArchive& operator<<(FArchive& Ar, FProfilerCycleGraph& Data)
  */
 struct FProfilerDataFrame
 {
-	uint32 Frame;
-	double FrameStart;
 	TMap<uint32, TArray<FProfilerCycleCounter> > CycleCounters;
 	TMap<uint32, FProfilerCycleGraph> CycleGraphs;
 	TArray<FProfilerFloatAccumulator> FloatAccumulators;
 	TArray<FProfilerCountAccumulator> CountAccumulators;
+
+	double FrameStart;
+	uint32 Frame;
 	bool MetaDataUpdated;
 };
 
diff --git a/Engine/Source/Programs/CrashReporter/CompressionHelper/Compression.cpp b/Engine/Source/Programs/CrashReporter/CompressionHelper/Compression.cpp
new file mode 100644
--- /dev/null
+++ b/Engine/Source/Programs/CrashReporter/CompressionHelper/Compression.cpp
@@ -0,0 +1,49 @@
+// Copyright 1998-2015 Epic Games, Inc. All Rights Reserved.
+
+#include "../../../ThirdParty/zlib/zlib-1.2.5/Inc/zlib.h"
+#pragma comment( lib, "../../../ThirdParty/zlib/zlib-1.2.5/Lib/Win64/zlib_64.lib" )
+
+// Taken from GenericPlatform.h
+
+// Unsigned base types.
+typedef unsigned char 		uint8;		// 8-bit  unsigned.
+typedef unsigned short int	uint16;		// 16-bit unsigned.
+typedef unsigned int		uint32;		// 32-bit unsigned.
+typedef unsigned long long	uint64;		// 64-bit unsigned.
+
+// Signed base types.
+typedef	signed char			int8;		// 8-bit  signed.
+typedef signed short int	int16;		// 16-bit signed.
+typedef signed int	 		int32;		// 32-bit signed.
+typedef signed long long	int64;		// 64-bit signed.
+
+// Taken from Compression.cpp
+
+/**
+ * Thread-safe abstract compression routine. Compresses memory from uncompressed buffer and writes it to compressed
+ * buffer. Updates CompressedSize with size of compressed data.
+ *
+ * @param	UncompressedBuffer			Buffer containing uncompressed data
+ * @param	UncompressedSize			Size of uncompressed data in bytes
+ * @param	CompressedBuffer			Buffer compressed data is going to be read from
+ * @param	CompressedSize				Size of CompressedBuffer data in bytes
+ * @return true if compression succeeds, false if it fails because CompressedBuffer was too small or other reasons
+ */
+extern "C"
+{
+
+	__declspec(dllexport)
+	bool __cdecl UE4UncompressMemoryZLIB( void* UncompressedBuffer, int32 UncompressedSize, const void* CompressedBuffer, int32 CompressedSize )
+	{
+		// Zlib wants to use unsigned long.
+		unsigned long ZCompressedSize	= CompressedSize;
+		unsigned long ZUncompressedSize	= UncompressedSize;
+	
+		// Uncompress data.
+		bool bOperationSucceeded = uncompress( (uint8*) UncompressedBuffer, &ZUncompressedSize, (const uint8*) CompressedBuffer, ZCompressedSize ) == Z_OK ? true : false;
+
+		// Sanity check to make sure we uncompressed as much data as we expected to.
+		return bOperationSucceeded;
+	}
+
+}
\ No newline at end of file
diff --git a/Engine/Source/Programs/CrashReporter/CrashReportCommon/Common.cs b/Engine/Source/Programs/CrashReporter/CrashReportCommon/Common.cs
new file mode 100644
--- /dev/null
+++ b/Engine/Source/Programs/CrashReporter/CrashReportCommon/Common.cs
@@ -0,0 +1,246 @@
+﻿// Copyright 1998-2015 Epic Games, Inc. All Rights Reserved.
+
+using System;
+using System.IO;
+using System.Xml.Serialization;
+using System.Collections;
+using System.Collections.Generic;
+using Tools.DotNETCommon.XmlHandler;
+using System.Xml;
+using System.Linq;
+using System.Security;
+
+namespace Tools.CrashReporter.CrashReportCommon
+{
+	/// <summary>
+	/// A class to return the result of the CheckReport web request.
+	/// </summary>
+	public class CrashReporterResult
+	{
+		/// <summary>A bool representing success or failure of the CheckReport web request.</summary>
+		[XmlAttribute]
+		public bool bSuccess = false;
+
+		/// <summary>An optional integer value typically representing a unique ID.</summary>
+		[XmlAttribute]
+		public int ID = 0;
+
+		/// <summary>A optional message to explain any return values.</summary>
+		[XmlAttribute]
+		public string Message = "";
+
+		/// <summary>A simple default constructor to allow Xml serialisation.</summary>
+		public CrashReporterResult()
+		{
+		}
+	}
+
+	/// <summary>
+	/// A class to send to the web service to check to see if the report has already been uploaded.
+	/// </summary>
+	public class CheckReportRequest
+	{
+		/// <summary>A unique ID, actually the directory name of the WER.</summary>
+		[XmlAttribute]
+		public string ReportId = "";
+
+		/// <summary>LEGACY SUPPORT for crashes reported from clients before CrashReportUploader had been rebuilt</summary>
+		[XmlAttribute]
+		public string DirectoryName = "";
+
+		/// <summary>A simple default constructor to allow Xml serialisation.</summary>
+		public CheckReportRequest()
+		{
+		}
+	}
+
+	/// <summary>
+	/// A class to send to the CrashReport website to add a new crash to the database.
+	/// </summary>
+	public class CrashDescription
+	{
+		/// <summary>The name of the branch this game was built out of.</summary>
+		[XmlElement]
+		public string BranchName = "";
+
+		/// <summary>The name of the game that crashed.</summary>
+		[XmlElement]
+		public string GameName = "";
+
+		/// <summary>The platform that crashed e.g. Win64.</summary>
+		[XmlElement]
+		public string Platform = "";
+
+		/// <summary>The mode the game was in e.g. editor.</summary>
+		[XmlElement]
+		public string EngineMode = "";
+
+		/// <summary>The four component version of the app e.g. 4.4.1.0</summary>
+		[XmlElement]
+		public string BuildVersion = "";
+
+		/// <summary> Changelist number. </summary>
+		[XmlElement]
+		public int BuiltFromCL = 0;
+
+		/// <summary>The command line of the application that crashed.</summary>
+		[XmlElement]
+		public string CommandLine = "";
+
+		/// <summary>The base directory where the app was running.</summary>
+		[XmlElement]
+		public string BaseDir = "";
+
+		/// <summary>The language code the app was running in e.g. 1033.</summary>
+		[XmlElement]
+		public string Language = "";
+
+		/// <summary>The language code of the system.</summary>
+		[XmlElement]
+		public string SystemLanguage = "";
+
+		/// <summary> 
+		/// The name of the user that caused this crash.
+		/// @UserName varchar(64)  
+		/// </summary>
+		[XmlElement]
+		public string UserName = "";
+
+		/// <summary>
+		/// The GUID of the machine the crash occurred on.
+		/// @ComputerName varchar(64) 
+		/// </summary>
+		[XmlElement]
+		public string MachineGuid = "";
+
+		/// <summary> 
+		/// The Epic account ID for the user who last used the Launcher.
+		/// @EpicAccountId	varchar(64)	
+		/// </summary>
+		[XmlElement]
+		public string EpicAccountId = "";
+
+		/// <summary>An array of strings representing the callstack of the crash.</summary>
+		[XmlElement]
+		public string[] CallStack = null;
+
+		/// <summary>An array of strings showing the source code around the crash.</summary>
+		[XmlElement]
+		public string[] SourceContext = null;
+
+		/// <summary>An array of strings representing the user description of the crash.</summary>
+		[XmlElement]
+		public string[] UserDescription = null;
+
+		/// <summary>The error message, can be assertion message, ensure message or message from the fatal error.</summary>
+		[XmlElement]
+		public string[] ErrorMessage = null;
+
+		/// <summary>The UTC time the crash occurred.</summary>
+		[XmlElement]
+		public DateTime TimeofCrash;
+
+		/// <summary>Whether this crash has a minidump.</summary>
+		[XmlElement]
+		public bool bHasMiniDump = false;
+
+		/// <summary>Whether this crash has a log file.</summary>
+		[XmlElement]
+		public bool bHasLog = false;
+
+		/// <summary>Obsolete.</summary>
+		[XmlElement]
+		public bool bHasDiags = false;
+
+		/// <summary>Obsolete.</summary>
+		[XmlElement]
+		public bool bHasWERData = false;
+
+		/// <summary>Whether this crash has a video.</summary>
+		[XmlElement]
+		public bool bHasVideo = false;
+
+		/// <summary> Whether the user allowed us to be contacted. </summary>
+		[XmlElement]
+		public bool bAllowToBeContacted = false;
+
+		/// <summary>A simple default constructor to allow Xml serialization.</summary>
+		public CrashDescription()
+		{
+		}
+	}
+
+	
+
+	/// <summary> Details about received compressed crash. </summary>
+	public class FCompressedCrashInformation
+	{
+		/// <summary> Size of the compressed data. </summary>
+		public string CompressedSize = "";
+
+		/// <summary> Size of the data after decompression. </summary>
+		public string UncompressedSize = "";
+
+		/// <summary> Number of files stored in the compressed data. </summary>
+		public string NumberOfFiles = "";
+
+		/// <summary>A simple default constructor to allow Xml serialization.</summary>
+		public FCompressedCrashInformation()
+		{}
+
+		/// <summary>Initialization constructor.</summary>
+		public FCompressedCrashInformation( string InCompressedSize, string InUncompressedSize, string InNumberOfFiles )
+		{
+			CompressedSize = InCompressedSize;
+			UncompressedSize = InUncompressedSize;
+			NumberOfFiles = InNumberOfFiles;
+		}
+
+		/// <summary> CompressedSize as int. </summary>
+		public int GetCompressedSize()
+		{
+			int Result = 0;
+			int.TryParse( CompressedSize, out Result );
+			return Result;
+		}
+
+		/// <summary> UncompressedSize as int. </summary>
+		public int GetUncompressedSize()
+		{
+			int Result = 0;
+			int.TryParse( UncompressedSize, out Result );
+			return Result;
+		}
+
+		/// <summary> NumberOfFiles as int. </summary>
+		public int GetNumberOfFiles()
+		{
+			int Result = 0;
+			int.TryParse( NumberOfFiles, out Result );
+			return Result;
+		}
+	}
+
+	/// <summary> Helper class for reading binary data. </summary>
+	public class FBinaryReaderHelper
+	{
+		/// <summary> Reads FString serialized with SerializeAsANSICharArray. </summary>
+		public static string ReadFixedSizeString( BinaryReader BinaryStream )
+		{
+			uint Length = BinaryStream.ReadUInt32();
+			string Result = new string( BinaryStream.ReadChars( (int)Length ) );
+			Result = FixFixedSizeString( Result );
+			return Result;
+		}
+
+		static String FixFixedSizeString( string String )
+		{
+			int RealLength = 0;
+			while( String[RealLength++] != '\0' )
+			{
+			}
+			string Result = String.Remove( RealLength - 1 );
+			return Result;
+		}
+	}
+}
