Copying //UE4/Dev-Sequencer to //UE4/Dev-Main (Source: //UE4/Dev-Sequencer @ 3178529)

diff --git a/Engine/Source/Developer/ProfilerClient/Private/ProfilerClientManager.cpp b/Engine/Source/Developer/ProfilerClient/Private/ProfilerClientManager.cpp
--- a/Engine/Source/Developer/ProfilerClient/Private/ProfilerClientManager.cpp
+++ b/Engine/Source/Developer/ProfilerClient/Private/ProfilerClientManager.cpp
@@ -1,32 +1,27 @@
 // Copyright 1998-2016 Epic Games, Inc. All Rights Reserved.
 
-/*=============================================================================
-	ProfilerClientModule.cpp: Implements the FProfilerClientModule class.
-=============================================================================*/
-
 #include "ProfilerClientPrivatePCH.h"
-#include "ISessionInstanceInfo.h"
-#include "SecureHash.h"
-#include "StatsData.h"
-#include "StatsFile.h"
-
-DEFINE_LOG_CATEGORY_STATIC( LogProfilerClient, Log, All );
-
-DECLARE_CYCLE_STAT( TEXT("HandleDataReceived"),	STAT_PC_HandleDataReceived,			STATGROUP_Profiler );
-DECLARE_CYCLE_STAT( TEXT("ReadStatMessages"),	STAT_PC_ReadStatMessages,			STATGROUP_Profiler );
-DECLARE_CYCLE_STAT( TEXT("AddStatMessages"),	STAT_PC_AddStatMessages,			STATGROUP_Profiler );
-DECLARE_CYCLE_STAT( TEXT("GenerateDataFrame"),	STAT_PC_GenerateDataFrame,			STATGROUP_Profiler );
-DECLARE_CYCLE_STAT( TEXT("AddStatFName"),		STAT_PC_AddStatFName,				STATGROUP_Profiler );
-DECLARE_CYCLE_STAT( TEXT("AddGroupFName"),		STAT_PC_AddGroupFName,				STATGROUP_Profiler );
-DECLARE_CYCLE_STAT( TEXT("GenerateCycleGraph"),	STAT_PC_GenerateCycleGraph,			STATGROUP_Profiler );
-DECLARE_CYCLE_STAT( TEXT("GenerateAccumulator"),STAT_PC_GenerateAccumulator,		STATGROUP_Profiler );
-DECLARE_CYCLE_STAT( TEXT("FindOrAddStat"),		STAT_PC_FindOrAddStat,				STATGROUP_Profiler );
-DECLARE_CYCLE_STAT( TEXT("FindOrAddThread"),	STAT_PC_FindOrAddThread,			STATGROUP_Profiler );
+#include "ProfilerClientManager.h"
+
+
+DEFINE_LOG_CATEGORY_STATIC(LogProfilerClient, Log, All);
+
+DECLARE_CYCLE_STAT(TEXT("HandleDataReceived"),	STAT_PC_HandleDataReceived,			STATGROUP_Profiler);
+DECLARE_CYCLE_STAT(TEXT("ReadStatMessages"),	STAT_PC_ReadStatMessages,			STATGROUP_Profiler);
+DECLARE_CYCLE_STAT(TEXT("AddStatMessages"),	STAT_PC_AddStatMessages,			STATGROUP_Profiler);
+DECLARE_CYCLE_STAT(TEXT("GenerateDataFrame"),	STAT_PC_GenerateDataFrame,			STATGROUP_Profiler);
+DECLARE_CYCLE_STAT(TEXT("AddStatFName"),		STAT_PC_AddStatFName,				STATGROUP_Profiler);
+DECLARE_CYCLE_STAT(TEXT("AddGroupFName"),		STAT_PC_AddGroupFName,				STATGROUP_Profiler);
+DECLARE_CYCLE_STAT(TEXT("GenerateCycleGraph"),	STAT_PC_GenerateCycleGraph,			STATGROUP_Profiler);
+DECLARE_CYCLE_STAT(TEXT("GenerateAccumulator"),STAT_PC_GenerateAccumulator,		STATGROUP_Profiler);
+DECLARE_CYCLE_STAT(TEXT("FindOrAddStat"),		STAT_PC_FindOrAddStat,				STATGROUP_Profiler);
+DECLARE_CYCLE_STAT(TEXT("FindOrAddThread"),	STAT_PC_FindOrAddThread,			STATGROUP_Profiler);
+
 
 /* FProfilerClientManager structors
  *****************************************************************************/
 
-FProfilerClientManager::FProfilerClientManager( const IMessageBusRef& InMessageBus )
+FProfilerClientManager::FProfilerClientManager(const TSharedRef<IMessageBus, ESPMode::ThreadSafe>& InMessageBus)
 {
 #if STATS
 	MessageBus = InMessageBus;
@@ -39,7 +34,7 @@ FProfilerClientManager::FProfilerClientManager( const IMessageBusRef& InMessageB
 
 	if (MessageEndpoint.IsValid())
 	{
-		OnShutdownMessageBusDelegateHandle = InMessageBus->OnShutdown().AddRaw( this, &FProfilerClientManager::HandleMessageBusShutdown );
+		OnShutdownMessageBusDelegateHandle = InMessageBus->OnShutdown().AddRaw(this, &FProfilerClientManager::HandleMessageBusShutdown);
 		MessageEndpoint->Subscribe<FProfilerServicePing>();
 	}
 
@@ -53,6 +48,7 @@ FProfilerClientManager::FProfilerClientManager( const IMessageBusRef& InMessageB
 #endif
 }
 
+
 FProfilerClientManager::~FProfilerClientManager()
 {
 #if STATS
@@ -62,7 +58,7 @@ FProfilerClientManager::~FProfilerClientManager()
 
 	if (MessageBus.IsValid())
 	{
-		MessageBus->OnShutdown().Remove( OnShutdownMessageBusDelegateHandle );
+		MessageBus->OnShutdown().Remove(OnShutdownMessageBusDelegateHandle);
 	}
 
 	LoadConnection = nullptr;
@@ -73,7 +69,7 @@ FProfilerClientManager::~FProfilerClientManager()
 /* IProfilerClient interface
  *****************************************************************************/
 
-void FProfilerClientManager::Subscribe( const FGuid& Session )
+void FProfilerClientManager::Subscribe(const FGuid& Session)
 {
 #if STATS
 	FGuid OldSessionId = ActiveSessionId;
@@ -99,11 +95,12 @@ void FProfilerClientManager::Subscribe( const FGuid& Session )
 
 	Connections.Reset();
 
-	UE_LOG( LogProfilerClient, Log, TEXT( "Subscribe Session: %s" ), *Session.ToString() );
+	UE_LOG(LogProfilerClient, Log, TEXT("Subscribe Session: %s"), *Session.ToString());
 #endif
 }
 
-void FProfilerClientManager::Track( const FGuid& Instance )
+
+void FProfilerClientManager::Track(const FGuid& Instance)
 {
 #if STATS
 	if (MessageEndpoint.IsValid() && ActiveSessionId.IsValid() && !PendingInstances.Contains(Instance))
@@ -115,12 +112,13 @@ void FProfilerClientManager::Track( const FGuid& Instance )
 		RetryTime = 5.f;
 		TickDelegateHandle = FTicker::GetCoreTicker().AddTicker(TickDelegate, RetryTime);
 
-		UE_LOG( LogProfilerClient, Verbose, TEXT( "Track Session: %s, Instance: %s" ), *ActiveSessionId.ToString(), *Instance.ToString() );
+		UE_LOG(LogProfilerClient, Verbose, TEXT("Track Session: %s, Instance: %s"), *ActiveSessionId.ToString(), *Instance.ToString());
 	}
 #endif
 }
 
-void FProfilerClientManager::Untrack( const FGuid& Instance )
+
+void FProfilerClientManager::Untrack(const FGuid& Instance)
 {
 #if STATS
 	if (MessageEndpoint.IsValid() && ActiveSessionId.IsValid())
@@ -131,11 +129,12 @@ void FProfilerClientManager::Untrack( const FGuid& Instance )
 		// fire the disconnection delegate
 		ProfilerClientDisconnectedDelegate.Broadcast(ActiveSessionId, Instance);
 
-		UE_LOG( LogProfilerClient, Verbose, TEXT( "Untrack Session: %s, Instance: %s" ), *ActiveSessionId.ToString(), *Instance.ToString() );
+		UE_LOG(LogProfilerClient, Verbose, TEXT("Untrack Session: %s, Instance: %s"), *ActiveSessionId.ToString(), *Instance.ToString());
 	}
 #endif
 }
 
+
 void FProfilerClientManager::Unsubscribe()
 {
 #if STATS
@@ -144,63 +143,66 @@ void FProfilerClientManager::Unsubscribe()
 #endif
 }
 
-void FProfilerClientManager::SetCaptureState( const bool bRequestedCaptureState, const FGuid& InstanceId /*= FGuid()*/  )
+
+void FProfilerClientManager::SetCaptureState(const bool bRequestedCaptureState, const FGuid& InstanceId /*= FGuid()*/ )
 {
 #if STATS
 	if (MessageEndpoint.IsValid() && ActiveSessionId.IsValid())
 	{
-		if( !InstanceId.IsValid() )
+		if(!InstanceId.IsValid())
 		{
 			TArray<FMessageAddress> Instances;
 			for (auto It = Connections.CreateConstIterator(); It; ++It)
 			{
-				Instances.Add( It.Value().ProfilerServiceAddress );
+				Instances.Add(It.Value().ProfilerServiceAddress);
 			}
 			MessageEndpoint->Send(new FProfilerServiceCapture(bRequestedCaptureState), Instances);
-			UE_LOG( LogProfilerClient, Verbose, TEXT( "SetCaptureState Session: %s, Instance: %s, State: %i" ), *ActiveSessionId.ToString(), *InstanceId.ToString(), (int32)bRequestedCaptureState );
+			UE_LOG(LogProfilerClient, Verbose, TEXT("SetCaptureState Session: %s, Instance: %s, State: %i"), *ActiveSessionId.ToString(), *InstanceId.ToString(), (int32)bRequestedCaptureState);
 		}
 		else
 		{
 			const FMessageAddress* MessageAddress = &Connections.Find(InstanceId)->ProfilerServiceAddress;
-			if( MessageAddress )
+			if(MessageAddress)
 			{
 				MessageEndpoint->Send(new FProfilerServiceCapture(bRequestedCaptureState), *MessageAddress);
 			}
 
-			UE_LOG( LogProfilerClient, Verbose, TEXT( "SetCaptureState Session: %s, Instance: %s, State: %i" ), *ActiveSessionId.ToString(), *InstanceId.ToString(), (int32)bRequestedCaptureState );
+			UE_LOG(LogProfilerClient, Verbose, TEXT("SetCaptureState Session: %s, Instance: %s, State: %i"), *ActiveSessionId.ToString(), *InstanceId.ToString(), (int32)bRequestedCaptureState);
 		}
 	}
 #endif
 }
 
-void FProfilerClientManager::SetPreviewState( const bool bRequestedPreviewState, const FGuid& InstanceId /*= FGuid()*/ )
+
+void FProfilerClientManager::SetPreviewState(const bool bRequestedPreviewState, const FGuid& InstanceId /*= FGuid()*/)
 {
 #if STATS
 	if (MessageEndpoint.IsValid() && ActiveSessionId.IsValid())
 	{
-		if( !InstanceId.IsValid() )
+		if(!InstanceId.IsValid())
 		{
 			TArray<FMessageAddress> Instances;
 			for (auto It = Connections.CreateConstIterator(); It; ++It)
 			{
-				Instances.Add( It.Value().ProfilerServiceAddress );
+				Instances.Add(It.Value().ProfilerServiceAddress);
 			}
 			MessageEndpoint->Send(new FProfilerServicePreview(bRequestedPreviewState), Instances);
-			UE_LOG( LogProfilerClient, Verbose, TEXT( "SetPreviewState Session: %s, Instance: %s, State: %i" ), *ActiveSessionId.ToString(), *InstanceId.ToString(), (int32)bRequestedPreviewState );
+			UE_LOG(LogProfilerClient, Verbose, TEXT("SetPreviewState Session: %s, Instance: %s, State: %i"), *ActiveSessionId.ToString(), *InstanceId.ToString(), (int32)bRequestedPreviewState);
 		}
 		else
 		{
 			const FMessageAddress* MessageAddress = &Connections.Find(InstanceId)->ProfilerServiceAddress;
-			if( MessageAddress )
+			if(MessageAddress)
 			{
 				MessageEndpoint->Send(new FProfilerServicePreview(bRequestedPreviewState), *MessageAddress);
 			}
-			UE_LOG( LogProfilerClient, Verbose, TEXT( "SetPreviewState Session: %s, Instance: %s, State: %i" ), *ActiveSessionId.ToString(), *InstanceId.ToString(), (int32)bRequestedPreviewState );
+			UE_LOG(LogProfilerClient, Verbose, TEXT("SetPreviewState Session: %s, Instance: %s, State: %i"), *ActiveSessionId.ToString(), *InstanceId.ToString(), (int32)bRequestedPreviewState);
 		}
 	}
 #endif
 }
 
+
 /*-----------------------------------------------------------------------------
 	New read test, still temporary, but around 4x faster
 -----------------------------------------------------------------------------*/
@@ -211,8 +213,9 @@ class FNewStatsReader : public FStatsReadFile
 	typedef FStatsReadFile Super;
 
 public:
+
 	/** Initialize. */
-	void Initialize( FProfilerClientManager* InProfilerClientManager, FServiceConnection* InLoadConnection )
+	void Initialize(FProfilerClientManager* InProfilerClientManager, FServiceConnection* InLoadConnection)
 	{
 		ProfilerClientManager = InProfilerClientManager;
 		LoadConnection = InLoadConnection;
@@ -221,19 +224,19 @@ public:
 protected:
 
 	/** Initialization constructor. */
-	FNewStatsReader( const TCHAR* InFilename )
-		: FStatsReadFile( InFilename, false )
-		, ProfilerClientManager( nullptr )
-		, LoadConnection( nullptr )
+	FNewStatsReader(const TCHAR* InFilename)
+		: FStatsReadFile(InFilename, false)
+		, ProfilerClientManager(nullptr)
+		, LoadConnection(nullptr)
 	{
 		// Keep only the last frame.
-		SetHistoryFrames( 1 );
+		SetHistoryFrames(1);
 	}	
 
 	/** Called every each frame has been read from the file. */
-	virtual void ReadStatsFrame( const TArray<FStatMessage>& CondensedMessages, const int64 Frame ) override
+	virtual void ReadStatsFrame(const TArray<FStatMessage>& CondensedMessages, const int64 Frame) override
 	{
-		SCOPE_CYCLE_COUNTER( STAT_PC_GenerateDataFrame );
+		SCOPE_CYCLE_COUNTER(STAT_PC_GenerateDataFrame);
 
 		FProfilerDataFrame& DataFrame = LoadConnection->CurrentData;
 
@@ -247,10 +250,10 @@ protected:
 		// Get the stat stack root and the non frame stats.
 		FRawStatStackNode Stack;
 		TArray<FStatMessage> NonFrameStats;
-		State.UncondenseStackStats( CondensedMessages, Stack, nullptr, &NonFrameStats );
+		State.UncondenseStackStats(CondensedMessages, Stack, nullptr, &NonFrameStats);
 
-		LoadConnection->GenerateCycleGraphs( Stack, DataFrame.CycleGraphs );
-		LoadConnection->GenerateAccumulators( NonFrameStats, DataFrame.CountAccumulators, DataFrame.FloatAccumulators );
+		LoadConnection->GenerateCycleGraphs(Stack, DataFrame.CycleGraphs);
+		LoadConnection->GenerateAccumulators(NonFrameStats, DataFrame.CountAccumulators, DataFrame.FloatAccumulators);
 
 		// Create a copy of the stats metadata.
 		FStatMetaData* MetaDataPtr = nullptr;
@@ -268,7 +271,7 @@ protected:
 		// Send to game thread.
 		FSimpleDelegateGraphTask::CreateAndDispatchWhenReady
 		(
-			FSimpleDelegateGraphTask::FDelegate::CreateRaw( ProfilerClientManager, &FProfilerClientManager::SendProfilerDataFrameToGame, DataFramePtr, MetaDataPtr, LoadConnection->InstanceId ),
+			FSimpleDelegateGraphTask::FDelegate::CreateRaw(ProfilerClientManager, &FProfilerClientManager::SendProfilerDataFrameToGame, DataFramePtr, MetaDataPtr, LoadConnection->InstanceId),
 			TStatId(), nullptr, ENamedThreads::GameThread
 		);
 	}
@@ -278,7 +281,7 @@ protected:
 	{
 		FSimpleDelegateGraphTask::CreateAndDispatchWhenReady
 		(
-			FSimpleDelegateGraphTask::FDelegate::CreateRaw( ProfilerClientManager, &FProfilerClientManager::FinalizeLoading, LoadConnection->InstanceId ),
+			FSimpleDelegateGraphTask::FDelegate::CreateRaw(ProfilerClientManager, &FProfilerClientManager::FinalizeLoading, LoadConnection->InstanceId),
 			TStatId(), NULL, ENamedThreads::GameThread
 		);
 	}
@@ -287,19 +290,19 @@ protected:
 	FServiceConnection* LoadConnection;
 };
 
-void FServiceConnection::LoadCapture( const FString& DataFilepath, FProfilerClientManager* ProfilerClientManager )
+
+void FServiceConnection::LoadCapture(const FString& DataFilepath, FProfilerClientManager* ProfilerClientManager)
 {
-	StatsReader = FStatsReader<FNewStatsReader>::Create( *DataFilepath );
+	StatsReader = FStatsReader<FNewStatsReader>::Create(*DataFilepath);
 	if (StatsReader)
 	{
-		StatsReader->Initialize( ProfilerClientManager, this );
+		StatsReader->Initialize(ProfilerClientManager, this);
 		StatsReader->ReadAndProcessAsynchronously();
 	}
 }
 
 
-
-void FProfilerClientManager::LoadCapture( const FString& DataFilepath, const FGuid& ProfileId )
+void FProfilerClientManager::LoadCapture(const FString& DataFilepath, const FGuid& ProfileId)
 {
 #if STATS
 	// Start an async load.
@@ -307,8 +310,8 @@ void FProfilerClientManager::LoadCapture( const FString& DataFilepath, const FGu
 	LoadConnection->InstanceId = ProfileId;
 	LoadConnection->StatMetaData.SecondsPerCycle = FPlatformTime::GetSecondsPerCycle(); // fix this by adding a message which specifies this
 
-	ProfilerLoadStartedDelegate.Broadcast( ProfileId );
-	LoadConnection->LoadCapture( DataFilepath, this );
+	ProfilerLoadStartedDelegate.Broadcast(ProfileId);
+	LoadConnection->LoadCapture(DataFilepath, this);
 
 	RetryTime = 0.05f;
 	TickDelegateHandle = FTicker::GetCoreTicker().AddTicker(TickDelegate, RetryTime);	
@@ -316,24 +319,24 @@ void FProfilerClientManager::LoadCapture( const FString& DataFilepath, const FGu
 }
 
 
-void FProfilerClientManager::RequestLastCapturedFile( const FGuid& InstanceId /*= FGuid()*/ )
+void FProfilerClientManager::RequestLastCapturedFile(const FGuid& InstanceId /*= FGuid()*/)
 {
 #if STATS
 	if (MessageEndpoint.IsValid() && ActiveSessionId.IsValid())
 	{
-		if( !InstanceId.IsValid() )
+		if(!InstanceId.IsValid())
 		{
 			TArray<FMessageAddress> Instances;
 			for (auto It = Connections.CreateConstIterator(); It; ++It)
 			{
-				Instances.Add( It.Value().ProfilerServiceAddress );
+				Instances.Add(It.Value().ProfilerServiceAddress);
 			}
 			MessageEndpoint->Send(new FProfilerServiceRequest(EProfilerRequestType::PRT_SendLastCapturedFile), Instances);
 		}
 		else
 		{
 			const FMessageAddress* MessageAddress = &Connections.Find(InstanceId)->ProfilerServiceAddress;
-			if( MessageAddress )
+			if(MessageAddress)
 			{
 				MessageEndpoint->Send(new FProfilerServiceRequest(EProfilerRequestType::PRT_SendLastCapturedFile), *MessageAddress);
 			}
@@ -342,6 +345,61 @@ void FProfilerClientManager::RequestLastCapturedFile( const FGuid& InstanceId /*
 #endif
 }
 
+
+const FStatMetaData& FProfilerClientManager::GetStatMetaData(const FGuid& InstanceId) const
+{
+	return Connections.Find(InstanceId)->StatMetaData;
+}
+
+
+FProfilerClientDataDelegate& FProfilerClientManager::OnProfilerData()
+{
+	return ProfilerDataDelegate;
+}
+
+
+FProfilerFileTransferDelegate& FProfilerClientManager::OnProfilerFileTransfer()
+{
+	return ProfilerFileTransferDelegate;
+}
+
+
+FProfilerClientConnectedDelegate& FProfilerClientManager::OnProfilerClientConnected()
+{
+	return ProfilerClientConnectedDelegate;
+}
+
+
+FProfilerClientDisconnectedDelegate& FProfilerClientManager::OnProfilerClientDisconnected()
+{
+	return ProfilerClientDisconnectedDelegate;
+}
+
+
+FProfilerMetaDataUpdateDelegate& FProfilerClientManager::OnMetaDataUpdated()
+{
+	return ProfilerMetaDataUpdatedDelegate;
+}
+
+
+FProfilerLoadStartedDelegate& FProfilerClientManager::OnLoadStarted()
+{
+	return ProfilerLoadStartedDelegate;
+}
+
+
+FProfilerLoadCompletedDelegate& FProfilerClientManager::OnLoadCompleted()
+{
+	return ProfilerLoadCompletedDelegate;
+}
+
+
+FProfilerLoadCancelledDelegate& FProfilerClientManager::OnLoadCancelled()
+{
+	return ProfilerLoadCancelledDelegate;
+}
+
+
 /* FProfilerClientManager event handlers
  *****************************************************************************/
 
@@ -355,7 +413,8 @@ void FProfilerClientManager::HandleMessageBusShutdown()
 #endif
 }
 
-void FProfilerClientManager::HandleServiceAuthorizeMessage( const FProfilerServiceAuthorize& Message, const IMessageContextRef& Context )
+
+void FProfilerClientManager::HandleServiceAuthorizeMessage(const FProfilerServiceAuthorize& Message, const TSharedRef<IMessageContext, ESPMode::ThreadSafe>& Context)
 {
 #if STATS
 	if (ActiveSessionId == Message.SessionId && PendingInstances.Contains(Message.InstanceId))
@@ -370,23 +429,23 @@ void FProfilerClientManager::HandleServiceAuthorizeMessage( const FProfilerServi
 		// Fire the client connection event
 		ProfilerClientConnectedDelegate.Broadcast(ActiveSessionId, Message.InstanceId);
 
-		UE_LOG( LogProfilerClient, Verbose, TEXT( "Authorize SessionId: %s, InstanceId: %s" ), *Message.SessionId.ToString(), *Message.InstanceId.ToString() );
+		UE_LOG(LogProfilerClient, Verbose, TEXT("Authorize SessionId: %s, InstanceId: %s"), *Message.SessionId.ToString(), *Message.InstanceId.ToString());
 	}
 #endif
 }
 
-/*-----------------------------------------------------------------------------
-	FServiceConnection
------------------------------------------------------------------------------*/
+
+/* FServiceConnection
+ *****************************************************************************/
 
 FServiceConnection::FServiceConnection() 
-	: StatsReader( nullptr )
-{
-}
+	: StatsReader(nullptr)
+{ }
+
+
+FServiceConnection::FServiceConnection(const FServiceConnection& InConnnection)
+{ }
 
-FServiceConnection::FServiceConnection( const FServiceConnection& InConnnection )
-{
-}
 
 FServiceConnection::~FServiceConnection()
 {
@@ -396,8 +455,8 @@ FServiceConnection::~FServiceConnection()
 
 		while (StatsReader->IsBusy())
 		{
-			FPlatformProcess::Sleep( 2.0f );
-			UE_LOG( LogStats, Log, TEXT( "RequestStop: Stage: %s / %3i%%" ), *StatsReader->GetProcessingStageAsString(), StatsReader->GetStageProgress() );
+			FPlatformProcess::Sleep(2.0f);
+			UE_LOG(LogStats, Log, TEXT("RequestStop: Stage: %s / %3i%%"), *StatsReader->GetProcessingStageAsString(), StatsReader->GetStageProgress());
 		}
 
 		delete StatsReader;
@@ -410,7 +469,8 @@ FServiceConnection::~FServiceConnection()
 	}
 }
 
-void FServiceConnection::Initialize( const FProfilerServiceAuthorize& Message, const IMessageContextRef& Context )
+
+void FServiceConnection::Initialize(const FProfilerServiceAuthorize& Message, const TSharedRef<IMessageContext, ESPMode::ThreadSafe>& Context)
 {
 #if STATS
 	ProfilerServiceAddress = Context->GetSender();
@@ -419,7 +479,8 @@ void FServiceConnection::Initialize( const FProfilerServiceAuthorize& Message, c
 #endif
 }
 
-bool FProfilerClientManager::CheckHashAndWrite( const FProfilerServiceFileChunk& FileChunk, const FProfilerFileChunkHeader& FileChunkHeader, FArchive* Writer )
+
+bool FProfilerClientManager::CheckHashAndWrite(const FProfilerServiceFileChunk& FileChunk, const FProfilerFileChunkHeader& FileChunkHeader, FArchive* Writer)
 {
 #if STATS
 	const int32 HashSize = 20;
@@ -428,27 +489,27 @@ bool FProfilerClientManager::CheckHashAndWrite( const FProfilerServiceFileChunk&
 	// De-hex string into TArray<uint8>
 	TArray<uint8> FileChunkData;
 	const int32 DataLength = FileChunk.HexData.Len() / 2;
-	FileChunkData.Reset( DataLength );
-	FileChunkData.AddUninitialized( DataLength );
-	FString::ToHexBlob( FileChunk.HexData, FileChunkData.GetData(), DataLength );
+	FileChunkData.Reset(DataLength);
+	FileChunkData.AddUninitialized(DataLength);
+	FString::ToHexBlob(FileChunk.HexData, FileChunkData.GetData(), DataLength);
 
 	// Hash file chunk data. 
 	FSHA1 Sha;
-	Sha.Update( FileChunkData.GetData(), FileChunkHeader.ChunkSize );
+	Sha.Update(FileChunkData.GetData(), FileChunkHeader.ChunkSize);
 	// Hash file chunk header.
-	Sha.Update( FileChunk.Header.GetData(), FileChunk.Header.Num() );
+	Sha.Update(FileChunk.Header.GetData(), FileChunk.Header.Num());
 	Sha.Final();
-	Sha.GetHash( LocalHash );
+	Sha.GetHash(LocalHash);
 
-	const int32 MemDiff = FMemory::Memcmp( FileChunk.ChunkHash.GetData(), LocalHash, HashSize );
+	const int32 MemDiff = FMemory::Memcmp(FileChunk.ChunkHash.GetData(), LocalHash, HashSize);
 
 	bool bResult = false;
 
-	if( MemDiff == 0 )
+	if(MemDiff == 0)
 	{
 		// Write the data to the archive.
-		Writer->Seek( FileChunkHeader.ChunkOffset );
-		Writer->Serialize( (void*)FileChunkData.GetData(), FileChunkHeader.ChunkSize );
+		Writer->Seek(FileChunkHeader.ChunkOffset);
+		Writer->Serialize((void*)FileChunkData.GetData(), FileChunkHeader.ChunkSize);
 
 		bResult = true;
 	}
@@ -459,54 +520,55 @@ bool FProfilerClientManager::CheckHashAndWrite( const FProfilerServiceFileChunk&
 #endif
 }
 
-void FProfilerClientManager::HandleServiceFileChunk( const FProfilerServiceFileChunk& FileChunk, const IMessageContextRef& Context )
+
+void FProfilerClientManager::HandleServiceFileChunk(const FProfilerServiceFileChunk& FileChunk, const TSharedRef<IMessageContext, ESPMode::ThreadSafe>& Context)
 {
 #if STATS
 	const TCHAR* StrTmp = TEXT(".tmp");
 
 	// Read file chunk header.
-	FMemoryReader Reader( FileChunk.Header );
+	FMemoryReader Reader(FileChunk.Header);
 	FProfilerFileChunkHeader FileChunkHeader;
 	Reader << FileChunkHeader;
 	FileChunkHeader.Validate();
 
-	const bool bValidFileChunk = !FailedTransfer.Contains( FileChunk.Filename );
+	const bool bValidFileChunk = !FailedTransfer.Contains(FileChunk.Filename);
 
 	if (ActiveSessionId.IsValid() && Connections.Find(FileChunk.InstanceId) != nullptr && bValidFileChunk)
 	{
-		FReceivedFileInfo* ReceivedFileInfo = ActiveTransfers.Find( FileChunk.Filename );
-		if( !ReceivedFileInfo )
+		FReceivedFileInfo* ReceivedFileInfo = ActiveTransfers.Find(FileChunk.Filename);
+		if(!ReceivedFileInfo)
 		{
 			const FString PathName = FPaths::ProfilingDir() + TEXT("UnrealStats/Received/");
 			const FString StatFilepath = PathName + FileChunk.Filename + StrTmp;
 
-			UE_LOG( LogProfilerClient, Log, TEXT( "Opening stats file for service-client sending: %s" ), *StatFilepath );
+			UE_LOG(LogProfilerClient, Log, TEXT("Opening stats file for service-client sending: %s"), *StatFilepath);
 
 			FArchive* FileWriter = IFileManager::Get().CreateFileWriter(*StatFilepath);
-			if( !FileWriter )
+			if(!FileWriter)
 			{
-				UE_LOG( LogProfilerClient, Error, TEXT( "Could not open: %s" ), *StatFilepath );
+				UE_LOG(LogProfilerClient, Error, TEXT("Could not open: %s"), *StatFilepath);
 				return;
 			}
 
-			ReceivedFileInfo = &ActiveTransfers.Add( FileChunk.Filename, FReceivedFileInfo(FileWriter,0,StatFilepath) );
-			ProfilerFileTransferDelegate.Broadcast( FileChunk.Filename, ReceivedFileInfo->Progress, FileChunkHeader.FileSize );
+			ReceivedFileInfo = &ActiveTransfers.Add(FileChunk.Filename, FReceivedFileInfo(FileWriter,0,StatFilepath));
+			ProfilerFileTransferDelegate.Broadcast(FileChunk.Filename, ReceivedFileInfo->Progress, FileChunkHeader.FileSize);
 		}
 
 		const bool bSimulateBadFileChunk = true;//FMath::Rand() % 10 != 0;
-		const bool bSuccess = CheckHashAndWrite( FileChunk, FileChunkHeader, ReceivedFileInfo->FileWriter ) && bSimulateBadFileChunk;	
-		if( bSuccess )
+		const bool bSuccess = CheckHashAndWrite(FileChunk, FileChunkHeader, ReceivedFileInfo->FileWriter) && bSimulateBadFileChunk;	
+		if(bSuccess)
 		{
 			ReceivedFileInfo->Progress += FileChunkHeader.ChunkSize;
 			ReceivedFileInfo->Update();
 
-			if( ReceivedFileInfo->Progress == FileChunkHeader.FileSize )
+			if(ReceivedFileInfo->Progress == FileChunkHeader.FileSize)
 			{
 				// File has been successfully sent, so send this information to the profiler service.
-				if( MessageEndpoint.IsValid() )
+				if(MessageEndpoint.IsValid())
 				{
-					MessageEndpoint->Send( new FProfilerServiceFileChunk( FGuid(),FileChunk.Filename,FProfilerFileChunkHeader(0,0,0,EProfilerFileChunkType::FinalizeFile).AsArray() ), Context->GetSender() );
-					ProfilerFileTransferDelegate.Broadcast( FileChunk.Filename, ReceivedFileInfo->Progress, FileChunkHeader.FileSize );
+					MessageEndpoint->Send(new FProfilerServiceFileChunk(FGuid(),FileChunk.Filename,FProfilerFileChunkHeader(0,0,0,EProfilerFileChunkType::FinalizeFile).AsArray()), Context->GetSender());
+					ProfilerFileTransferDelegate.Broadcast(FileChunk.Filename, ReceivedFileInfo->Progress, FileChunkHeader.FileSize);
 				}
 				
 				// Delete the file writer.
@@ -514,31 +576,32 @@ void FProfilerClientManager::HandleServiceFileChunk( const FProfilerServiceFileC
 				ReceivedFileInfo->FileWriter = nullptr;
 
 				// Rename the stats file.
-				IFileManager::Get().Move( *ReceivedFileInfo->DestFilepath.Replace( StrTmp, TEXT("") ), *ReceivedFileInfo->DestFilepath );
+				IFileManager::Get().Move(*ReceivedFileInfo->DestFilepath.Replace(StrTmp, TEXT("")), *ReceivedFileInfo->DestFilepath);
 
-				ActiveTransfers.Remove( FileChunk.Filename );
+				ActiveTransfers.Remove(FileChunk.Filename);
 
-				UE_LOG( LogProfilerClient, Log, TEXT( "File service-client received successfully: %s" ), *FileChunk.Filename );
+				UE_LOG(LogProfilerClient, Log, TEXT("File service-client received successfully: %s"), *FileChunk.Filename);
 			}
 			else
 			{
-				ProfilerFileTransferDelegate.Broadcast( FileChunk.Filename, ReceivedFileInfo->Progress, FileChunkHeader.FileSize );			
+				ProfilerFileTransferDelegate.Broadcast(FileChunk.Filename, ReceivedFileInfo->Progress, FileChunkHeader.FileSize);			
 			}
 		}
 		else
 		{
 			// This chunk is a bad chunk, so ask for resending it.
-			if( MessageEndpoint.IsValid() )
+			if(MessageEndpoint.IsValid())
 			{
-				MessageEndpoint->Send( new FProfilerServiceFileChunk(FileChunk,FProfilerServiceFileChunk::FNullTag()), Context->GetSender() );
-				UE_LOG( LogProfilerClient, Log, TEXT( "Received a bad chunk of file, resending: %5i, %6u, %10u, %s" ), FileChunk.HexData.Len(), ReceivedFileInfo->Progress, FileChunkHeader.FileSize, *FileChunk.Filename );
+				MessageEndpoint->Send(new FProfilerServiceFileChunk(FileChunk,FProfilerServiceFileChunk::FNullTag()), Context->GetSender());
+				UE_LOG(LogProfilerClient, Log, TEXT("Received a bad chunk of file, resending: %5i, %6u, %10u, %s"), FileChunk.HexData.Len(), ReceivedFileInfo->Progress, FileChunkHeader.FileSize, *FileChunk.Filename);
 			}
 		}
 	}
 #endif
 }
 
-void FProfilerClientManager::HandleServicePingMessage( const FProfilerServicePing& Message, const IMessageContextRef& Context )
+
+void FProfilerClientManager::HandleServicePingMessage(const FProfilerServicePing& Message, const TSharedRef<IMessageContext, ESPMode::ThreadSafe>& Context)
 {
 #if STATS
 	if (MessageEndpoint.IsValid())
@@ -550,13 +613,13 @@ void FProfilerClientManager::HandleServicePingMessage( const FProfilerServicePin
 		}
 		MessageEndpoint->Send(new FProfilerServicePong(), Instances);
 
-		UE_LOG( LogProfilerClient, Verbose, TEXT( "Ping GetSender: %s" ), *Context->GetSender().ToString() );
+		UE_LOG(LogProfilerClient, Verbose, TEXT("Ping GetSender: %s"), *Context->GetSender().ToString());
 	}
 #endif
 }
 
 
-bool FProfilerClientManager::HandleTicker( float DeltaTime )
+bool FProfilerClientManager::HandleTicker(float DeltaTime)
 {
 #if STATS
 	if (PendingInstances.Num() > 0 && FDateTime::Now() > LastPingTime + DeltaTime)
@@ -576,7 +639,8 @@ bool FProfilerClientManager::HandleTicker( float DeltaTime )
 	return false;
 }
 
-bool FProfilerClientManager::HandleMessagesTicker( float DeltaTime )
+
+bool FProfilerClientManager::HandleMessagesTicker(float DeltaTime)
 {
 #if STATS
 	for (auto It = Connections.CreateIterator(); It; ++It)
@@ -584,41 +648,41 @@ bool FProfilerClientManager::HandleMessagesTicker( float DeltaTime )
 		FServiceConnection& Connection = It.Value();
 
 		TArray<int64> Frames;
-		Connection.ReceivedData.GenerateKeyArray( Frames );
+		Connection.ReceivedData.GenerateKeyArray(Frames);
 		Frames.Sort();
 
 		// MessageBus sends all data in out of order fashion.
 		// We buffer frame to make sure that all frames are received in the proper order.
 		const int32 NUM_BUFFERED_FRAMES = 15;
 
-		for( int32 Index = 0; Index < Frames.Num(); Index++ )
+		for(int32 Index = 0; Index < Frames.Num(); Index++)
 		{
 			if (Connection.ReceivedData.Num() < NUM_BUFFERED_FRAMES)
 			{
 				break;
 			}
 
-			//FScopeLogTime SLT( "HandleMessagesTicker" );
+			//FScopeLogTime SLT("HandleMessagesTicker");
 
 			const int64 FrameNum = Frames[Index];
-			const TArray<uint8>* const Data = Connection.ReceivedData.FindChecked( FrameNum );
+			const TArray<uint8>* const Data = Connection.ReceivedData.FindChecked(FrameNum);
 			FStatsReadStream& Stream = Connection.Stream;
 
 
 			// Read all messages from the uncompressed buffer.
-			FMemoryReader MemoryReader( *Data, true );
+			FMemoryReader MemoryReader(*Data, true);
 			while (MemoryReader.Tell() < MemoryReader.TotalSize())
 			{
 				// Read the message.
-				FStatMessage Message( Stream.ReadMessage( MemoryReader ) );
-				new (Connection.PendingStatMessagesMessages)FStatMessage( Message );
+				FStatMessage Message(Stream.ReadMessage(MemoryReader));
+				new (Connection.PendingStatMessagesMessages)FStatMessage(Message);
 			}
 
 			// Adds a new from from the pending messages, the pending messages will be removed after the call.
-			Connection.CurrentThreadState.ProcessMetaDataAndLeaveDataOnly( Connection.PendingStatMessagesMessages );
-			Connection.CurrentThreadState.AddFrameFromCondensedMessages( Connection.PendingStatMessagesMessages );
+			Connection.CurrentThreadState.ProcessMetaDataAndLeaveDataOnly(Connection.PendingStatMessagesMessages);
+			Connection.CurrentThreadState.AddFrameFromCondensedMessages(Connection.PendingStatMessagesMessages);
 
-			UE_LOG( LogProfilerClient, VeryVerbose, TEXT( "Frame=%i/%i, FNamesIndexMap=%i, CurrentMetadataSize=%i" ), FrameNum, Frames.Num(), Connection.Stream.FNamesIndexMap.Num(), Connection.CurrentThreadState.ShortNameToLongName.Num() );
+			UE_LOG(LogProfilerClient, VeryVerbose, TEXT("Frame=%i/%i, FNamesIndexMap=%i, CurrentMetadataSize=%i"), FrameNum, Frames.Num(), Connection.Stream.FNamesIndexMap.Num(), Connection.CurrentThreadState.ShortNameToLongName.Num());
 
 			// create an old format data frame from the data
 			Connection.GenerateProfilerDataFrame();
@@ -626,33 +690,33 @@ bool FProfilerClientManager::HandleMessagesTicker( float DeltaTime )
 			// Fire a meta data update message
 			if (Connection.CurrentData.MetaDataUpdated)
 			{
-				ProfilerMetaDataUpdatedDelegate.Broadcast( Connection.InstanceId, Connection.StatMetaData );
+				ProfilerMetaDataUpdatedDelegate.Broadcast(Connection.InstanceId, Connection.StatMetaData);
 			}
 
 			// send the data out
-			ProfilerDataDelegate.Broadcast( Connection.InstanceId, Connection.CurrentData );
+			ProfilerDataDelegate.Broadcast(Connection.InstanceId, Connection.CurrentData);
 
 			delete Data;
-			Connection.ReceivedData.Remove( FrameNum );
+			Connection.ReceivedData.Remove(FrameNum);
 		}
 	}
 
 	// Remove any active transfer that timed out.
-	for( auto It = ActiveTransfers.CreateIterator(); It; ++It )
+	for(auto It = ActiveTransfers.CreateIterator(); It; ++It)
 	{
 		FReceivedFileInfo& ReceivedFileInfo = It.Value();
 		const FString& Filename = It.Key();
 
-		if( ReceivedFileInfo.IsTimedOut() )
+		if(ReceivedFileInfo.IsTimedOut())
 		{
-			UE_LOG( LogProfilerClient, Log, TEXT( "File service-client timed out, aborted: %s" ), *Filename );
-			FailedTransfer.Add( Filename );
+			UE_LOG(LogProfilerClient, Log, TEXT("File service-client timed out, aborted: %s"), *Filename);
+			FailedTransfer.Add(Filename);
 
 			delete ReceivedFileInfo.FileWriter;
 			ReceivedFileInfo.FileWriter = nullptr;
 
-			IFileManager::Get().Delete( *ReceivedFileInfo.DestFilepath );
-			ProfilerFileTransferDelegate.Broadcast( Filename, -1, -1 );
+			IFileManager::Get().Delete(*ReceivedFileInfo.DestFilepath);
+			ProfilerFileTransferDelegate.Broadcast(Filename, -1, -1);
 			It.RemoveCurrent();	
 		}
 	}
@@ -661,99 +725,105 @@ bool FProfilerClientManager::HandleMessagesTicker( float DeltaTime )
 	return true;
 }
 
-void FProfilerClientManager::HandleServicePreviewAckMessage( const FProfilerServicePreviewAck& Message, const IMessageContextRef& Context )
+
+void FProfilerClientManager::HandleServicePreviewAckMessage(const FProfilerServicePreviewAck& Message, const TSharedRef<IMessageContext, ESPMode::ThreadSafe>& Context)
 {
 #if STATS
 	if (ActiveSessionId.IsValid() && Connections.Find(Message.InstanceId) != nullptr)
 	{
 		FServiceConnection& Connection = *Connections.Find(Message.InstanceId);
 
-		UE_LOG( LogProfilerClient, Verbose, TEXT( "PreviewAck InstanceId: %s, GetSender: %s" ), *Message.InstanceId.ToString(), *Context->GetSender().ToString() );
+		UE_LOG(LogProfilerClient, Verbose, TEXT("PreviewAck InstanceId: %s, GetSender: %s"), *Message.InstanceId.ToString(), *Context->GetSender().ToString());
 	}
 #endif
 }
 
-void FProfilerClientManager::HandleProfilerServiceData2Message( const FProfilerServiceData2& Message, const IMessageContextRef& Context )
+
+void FProfilerClientManager::HandleProfilerServiceData2Message(const FProfilerServiceData2& Message, const TSharedRef<IMessageContext, ESPMode::ThreadSafe>& Context)
 {
 #if STATS
 	SCOPE_CYCLE_COUNTER(STAT_PC_HandleDataReceived);
 	if (ActiveSessionId.IsValid() && Connections.Find(Message.InstanceId) != nullptr)
 	{
 		// Create a temporary profiler data and prepare all data.
-		FProfilerServiceData2* ToProcess = new FProfilerServiceData2( Message.InstanceId, Message.Frame, Message.HexData, Message.CompressedSize, Message.UncompressedSize );
+		FProfilerServiceData2* ToProcess = new FProfilerServiceData2(Message.InstanceId, Message.Frame, Message.HexData, Message.CompressedSize, Message.UncompressedSize);
 
 		// Decompression and decoding is done on the task graph.
 		FSimpleDelegateGraphTask::CreateAndDispatchWhenReady
 		(
-			FSimpleDelegateGraphTask::FDelegate::CreateRaw( this, &FProfilerClientManager::DecompressDataAndSendToGame, ToProcess ), 
+			FSimpleDelegateGraphTask::FDelegate::CreateRaw(this, &FProfilerClientManager::DecompressDataAndSendToGame, ToProcess), 
 			TStatId()
 		);
 	}
 #endif
 }
 
-void FProfilerClientManager::DecompressDataAndSendToGame( FProfilerServiceData2* ToProcess )
+
+void FProfilerClientManager::DecompressDataAndSendToGame(FProfilerServiceData2* ToProcess)
 {
-	DECLARE_SCOPE_CYCLE_COUNTER( TEXT( "FProfilerClientManager::DecompressDataAndSendToGame" ), STAT_FProfilerClientManager_DecompressDataAndSendToGame, STATGROUP_Profiler );
+	DECLARE_SCOPE_CYCLE_COUNTER(TEXT("FProfilerClientManager::DecompressDataAndSendToGame"), STAT_FProfilerClientManager_DecompressDataAndSendToGame, STATGROUP_Profiler);
 
 	// De-hex string into TArray<uint8>
 	TArray<uint8> CompressedData;
-	CompressedData.Reset( ToProcess->CompressedSize );
-	CompressedData.AddUninitialized( ToProcess->CompressedSize );
-	FString::ToHexBlob( ToProcess->HexData, CompressedData.GetData(), ToProcess->CompressedSize );
+	CompressedData.Reset(ToProcess->CompressedSize);
+	CompressedData.AddUninitialized(ToProcess->CompressedSize);
+	FString::ToHexBlob(ToProcess->HexData, CompressedData.GetData(), ToProcess->CompressedSize);
 
 	// Decompress data.
 	TArray<uint8> UncompressedData;
-	UncompressedData.Reset( ToProcess->UncompressedSize );
-	UncompressedData.AddUninitialized( ToProcess->UncompressedSize );
+	UncompressedData.Reset(ToProcess->UncompressedSize);
+	UncompressedData.AddUninitialized(ToProcess->UncompressedSize);
 
-	bool bResult = FCompression::UncompressMemory( COMPRESS_ZLIB, UncompressedData.GetData(), ToProcess->UncompressedSize, CompressedData.GetData(), ToProcess->CompressedSize );
-	check( bResult );
+	bool bResult = FCompression::UncompressMemory(COMPRESS_ZLIB, UncompressedData.GetData(), ToProcess->UncompressedSize, CompressedData.GetData(), ToProcess->CompressedSize);
+	check(bResult);
 
 	// Send to the game thread. Connections is not thread-safe, so we cannot add the data here.
-	TArray<uint8>* DateToGame = new TArray<uint8>( MoveTemp( UncompressedData ) );
+	TArray<uint8>* DateToGame = new TArray<uint8>(MoveTemp(UncompressedData));
 
 	FSimpleDelegateGraphTask::CreateAndDispatchWhenReady
 	(
-		FSimpleDelegateGraphTask::FDelegate::CreateRaw( this, &FProfilerClientManager::SendDataToGame, DateToGame, ToProcess->Frame, ToProcess->InstanceId ), 
+		FSimpleDelegateGraphTask::FDelegate::CreateRaw(this, &FProfilerClientManager::SendDataToGame, DateToGame, ToProcess->Frame, ToProcess->InstanceId), 
 		TStatId(), nullptr, ENamedThreads::GameThread
 	);
 
 	delete ToProcess;
 }
 
-void FProfilerClientManager::SendDataToGame( TArray<uint8>* DataToGame, int64 Frame, const FGuid InstanceId )
+
+void FProfilerClientManager::SendDataToGame(TArray<uint8>* DataToGame, int64 Frame, const FGuid InstanceId)
 {
-	if (ActiveSessionId.IsValid() && Connections.Find( InstanceId ) != nullptr)
+	if (ActiveSessionId.IsValid() && Connections.Find(InstanceId) != nullptr)
 	{
-		FServiceConnection& Connection = *Connections.Find( InstanceId );
+		FServiceConnection& Connection = *Connections.Find(InstanceId);
 
 		// Add the message to the connections queue.
-		UE_LOG( LogProfilerClient, VeryVerbose, TEXT( "Frame: %i, UncompressedSize: %i, InstanceId: %s" ), Frame, DataToGame->Num(), *InstanceId.ToString() );
-		Connection.ReceivedData.Add( Frame, DataToGame );
+		UE_LOG(LogProfilerClient, VeryVerbose, TEXT("Frame: %i, UncompressedSize: %i, InstanceId: %s"), Frame, DataToGame->Num(), *InstanceId.ToString());
+		Connection.ReceivedData.Add(Frame, DataToGame);
 	}
 }
 
-void FProfilerClientManager::SendProfilerDataFrameToGame( FProfilerDataFrame* NewData, FStatMetaData* MetaDataPtr, const FGuid InstanceId )
+
+void FProfilerClientManager::SendProfilerDataFrameToGame(FProfilerDataFrame* NewData, FStatMetaData* MetaDataPtr, const FGuid InstanceId)
 {
 	if (Connections.Find(InstanceId) != nullptr)
 	{
 		if (MetaDataPtr)
 		{
-			ProfilerMetaDataUpdatedDelegate.Broadcast( InstanceId, *MetaDataPtr );
+			ProfilerMetaDataUpdatedDelegate.Broadcast(InstanceId, *MetaDataPtr);
 			delete MetaDataPtr;
 			MetaDataPtr = nullptr;
 		}
 
 		if (NewData)
 		{
-			ProfilerDataDelegate.Broadcast( InstanceId, *NewData );
+			ProfilerDataDelegate.Broadcast(InstanceId, *NewData);
 			delete NewData;
 			NewData = nullptr;
 		}
 	}
 }
 
+
 void FProfilerClientManager::Shutdown()
 {
 	// Delete all active file writers and remove temporary files.
@@ -764,15 +834,16 @@ void FProfilerClientManager::Shutdown()
 		delete ReceivedFileInfo.FileWriter;
 		ReceivedFileInfo.FileWriter = nullptr;
 
-		IFileManager::Get().Delete( *ReceivedFileInfo.DestFilepath );
+		IFileManager::Get().Delete(*ReceivedFileInfo.DestFilepath);
 
-		UE_LOG( LogProfilerClient, Log, TEXT( "File service-client transfer aborted: %s" ), *It.Key() );
+		UE_LOG(LogProfilerClient, Log, TEXT("File service-client transfer aborted: %s"), *It.Key());
 	}
 
-	FTicker::GetCoreTicker().RemoveTicker( MessageDelegateHandle );
-	FTicker::GetCoreTicker().RemoveTicker( TickDelegateHandle );
+	FTicker::GetCoreTicker().RemoveTicker(MessageDelegateHandle);
+	FTicker::GetCoreTicker().RemoveTicker(TickDelegateHandle);
 }
 
+
 void FProfilerClientManager::FinalizeLoading(const FGuid InstanceId)
 {
 	if (Connections.Find(InstanceId) != nullptr)
@@ -788,6 +859,7 @@ void FProfilerClientManager::FinalizeLoading(const FGuid InstanceId)
 	}
 }
 
+
 void FProfilerClientManager::CancelLoading(const FGuid InstanceId)
 {
 	if (Connections.Find(InstanceId) != nullptr)
@@ -801,13 +873,14 @@ void FProfilerClientManager::CancelLoading(const FGuid InstanceId)
 	}
 }
 
+
 #if STATS
 
-int32 FServiceConnection::FindOrAddStat( const FStatNameAndInfo& StatNameAndInfo, uint32 StatType)
+int32 FServiceConnection::FindOrAddStat(const FStatNameAndInfo& StatNameAndInfo, uint32 StatType)
 {
 	SCOPE_CYCLE_COUNTER(STAT_PC_FindOrAddStat);
 	const FName LongName = StatNameAndInfo.GetRawName();
-	int32* const StatIDPtr = LongNameToStatID.Find( LongName );
+	int32* const StatIDPtr = LongNameToStatID.Find(LongName);
 	int32 StatID = StatIDPtr != nullptr ? *StatIDPtr : -1;
 
 	if (!StatIDPtr)
@@ -827,7 +900,7 @@ int32 FServiceConnection::FindOrAddStat( const FStatNameAndInfo& StatNameAndInfo
 		else if (StatName == FStatConstants::NAME_ThreadRoot)
 		{
 			StatID = LongNameToStatID.Add(LongName, 1);
-			GroupName = TEXT( "NoGroup" );
+			GroupName = TEXT("NoGroup");
 		}
 		else
 		{
@@ -839,20 +912,20 @@ int32 FServiceConnection::FindOrAddStat( const FStatNameAndInfo& StatNameAndInfo
 		FStatDescription StatDescription;
 		StatDescription.ID = StatID;
 		StatDescription.Name = !Description.IsEmpty() ? Description : StatName.ToString();
-		if( StatDescription.Name.Contains( TEXT("STAT_") ) )
+		if(StatDescription.Name.Contains(TEXT("STAT_")))
 		{
 			StatDescription.Name = StatDescription.Name.RightChop(FString(TEXT("STAT_")).Len());
 		}
 		StatDescription.StatType = StatType;
 
-		if( GroupName == NAME_None && Stream.Header.Version == EStatMagicNoHeader::NO_VERSION )
+		if(GroupName == NAME_None && Stream.Header.Version == EStatMagicNoHeader::NO_VERSION)
 		{	
 			// @todo Add more ways to group the stats.
-			const int32 Thread_Pos = StatDescription.Name.Find( TEXT("Thread_") );
-			const int32 _0Pos = StatDescription.Name.Find( TEXT("_0") );
+			const int32 Thread_Pos = StatDescription.Name.Find(TEXT("Thread_"));
+			const int32 _0Pos = StatDescription.Name.Find(TEXT("_0"));
 			const bool bIsThread = Thread_Pos != INDEX_NONE && _0Pos > Thread_Pos;
 			// Add a special group for all threads.
-			if( bIsThread )
+			if(bIsThread)
 			{
 				GroupName = TEXT("Threads");
 			}
@@ -863,13 +936,13 @@ int32 FServiceConnection::FindOrAddStat( const FStatNameAndInfo& StatNameAndInfo
 			}
 		}
 
-		int32* const GroupIDPtr = GroupNameArray.Find( GroupName );
+		int32* const GroupIDPtr = GroupNameArray.Find(GroupName);
 		int32 GroupID = GroupIDPtr != nullptr ? *GroupIDPtr : -1;
-		if( !GroupIDPtr )
+		if(!GroupIDPtr)
 		{
 			// add a new group description to the meta data
 			GroupID = GroupNameArray.Add(GroupName, GroupNameArray.Num()+10);
-			check( GroupID != -1 );
+			check(GroupID != -1);
 
 			FStatGroupDescription GroupDescription;
 			GroupDescription.ID = GroupID;
@@ -887,19 +960,20 @@ int32 FServiceConnection::FindOrAddStat( const FStatNameAndInfo& StatNameAndInfo
 	return StatID;
 }
 
+
 int32 FServiceConnection::FindOrAddThread(const FStatNameAndInfo& Thread)
 {
 	SCOPE_CYCLE_COUNTER(STAT_PC_FindOrAddThread);
 
 	// The description of a thread group contains the thread id
 	const FString Desc = Thread.GetDescription();
-	const uint32 ThreadID = FStatsUtils::ParseThreadID( Desc );
+	const uint32 ThreadID = FStatsUtils::ParseThreadID(Desc);
 
 	const FName ShortName = Thread.GetShortName();
 
 	// add to the meta data
 	const int32 OldNum = StatMetaData.ThreadDescriptions.Num();
-	StatMetaData.ThreadDescriptions.Add( ThreadID, ShortName.ToString() );
+	StatMetaData.ThreadDescriptions.Add(ThreadID, ShortName.ToString());
 	const int32 NewNum = StatMetaData.ThreadDescriptions.Num();
 
 	// meta data has been updated
@@ -908,6 +982,7 @@ int32 FServiceConnection::FindOrAddThread(const FStatNameAndInfo& Thread)
 	return ThreadID;
 }
 
+
 void FServiceConnection::GenerateAccumulators(TArray<FStatMessage>& Stats, TArray<FProfilerCountAccumulator>& CountAccumulators, TArray<FProfilerFloatAccumulator>& FloatAccumulators)
 {
 	SCOPE_CYCLE_COUNTER(STAT_PC_GenerateAccumulator)
@@ -918,11 +993,11 @@ void FServiceConnection::GenerateAccumulators(TArray<FStatMessage>& Stats, TArra
 		uint32 StatType = STATTYPE_Error;
 		if (StatMessage.NameAndInfo.GetField<EStatDataType>() == EStatDataType::ST_int64)
 		{
-			if (StatMessage.NameAndInfo.GetFlag( EStatMetaFlags::IsCycle ))
+			if (StatMessage.NameAndInfo.GetFlag(EStatMetaFlags::IsCycle))
 			{
 				StatType = STATTYPE_CycleCounter;
 			}
-			else if (StatMessage.NameAndInfo.GetFlag( EStatMetaFlags::IsMemory ))
+			else if (StatMessage.NameAndInfo.GetFlag(EStatMetaFlags::IsMemory))
 			{
 				StatType = STATTYPE_MemoryCounter;
 			}
@@ -938,7 +1013,7 @@ void FServiceConnection::GenerateAccumulators(TArray<FStatMessage>& Stats, TArra
 
 		if (StatType != STATTYPE_Error)
 		{
-			const int32 StatId = FindOrAddStat( StatMessage.NameAndInfo, StatType );
+			const int32 StatId = FindOrAddStat(StatMessage.NameAndInfo, StatType);
 
 			if (StatMessage.NameAndInfo.GetField<EStatDataType>() == EStatDataType::ST_int64)
 			{
@@ -946,7 +1021,7 @@ void FServiceConnection::GenerateAccumulators(TArray<FStatMessage>& Stats, TArra
 				FProfilerCountAccumulator Data;
 				Data.StatId = StatId;
 				Data.Value = StatMessage.GetValue_int64();
-				CountAccumulators.Add( Data );
+				CountAccumulators.Add(Data);
 			}
 			else if (StatMessage.NameAndInfo.GetField<EStatDataType>() == EStatDataType::ST_double)
 			{
@@ -954,7 +1029,7 @@ void FServiceConnection::GenerateAccumulators(TArray<FStatMessage>& Stats, TArra
 				FProfilerFloatAccumulator Data;
 				Data.StatId = StatId;
 				Data.Value = StatMessage.GetValue_double();
-				FloatAccumulators.Add( Data );
+				FloatAccumulators.Add(Data);
 
 				const FName StatName = StatMessage.NameAndInfo.GetRawName();
 				if (StatName == FStatConstants::RAW_SecondsPerCycle)
@@ -966,6 +1041,7 @@ void FServiceConnection::GenerateAccumulators(TArray<FStatMessage>& Stats, TArra
 	}
 }
 
+
 void FServiceConnection::CreateGraphRecursively(const FRawStatStackNode* Root, FProfilerCycleGraph& Graph, uint32 InStartCycles)
 {
 	Graph.FrameStart = InStartCycles;
@@ -989,8 +1065,8 @@ void FServiceConnection::CreateGraphRecursively(const FRawStatStackNode* Root, F
 	uint32 ChildStartCycles = InStartCycles;
 	TArray<FRawStatStackNode*> ChildArray;
 	Root->Children.GenerateValueArray(ChildArray);
-	ChildArray.Sort( FStatDurationComparer<FRawStatStackNode>() );
-	for( int32 Index = 0; Index < ChildArray.Num(); ++Index )
+	ChildArray.Sort(FStatDurationComparer<FRawStatStackNode>());
+	for(int32 Index = 0; Index < ChildArray.Num(); ++Index)
 	{
 		const FRawStatStackNode* ChildStat = ChildArray[Index];
 
@@ -1007,6 +1083,7 @@ void FServiceConnection::CreateGraphRecursively(const FRawStatStackNode* Root, F
 	}
 }
 
+
 void FServiceConnection::GenerateCycleGraphs(const FRawStatStackNode& Root, TMap<uint32, FProfilerCycleGraph>& CycleGraphs)
 {
 	SCOPE_CYCLE_COUNTER(STAT_PC_GenerateCycleGraph);
@@ -1033,9 +1110,10 @@ void FServiceConnection::GenerateCycleGraphs(const FRawStatStackNode& Root, TMap
 	}
 }
 
+
 void FServiceConnection::GenerateProfilerDataFrame()
 {
-	SCOPE_CYCLE_COUNTER( STAT_PC_GenerateDataFrame );
+	SCOPE_CYCLE_COUNTER(STAT_PC_GenerateDataFrame);
 	FProfilerDataFrame& DataFrame = CurrentData;
 	DataFrame.Frame = CurrentThreadState.CurrentGameFrame;
 	DataFrame.FrameStart = 0.0;
@@ -1047,13 +1125,14 @@ void FServiceConnection::GenerateProfilerDataFrame()
 	// get the stat stack root and the non frame stats
 	FRawStatStackNode Stack;
 	TArray<FStatMessage> NonFrameStats;
-	CurrentThreadState.UncondenseStackStats( CurrentThreadState.CurrentGameFrame, Stack, nullptr, &NonFrameStats );
+	CurrentThreadState.UncondenseStackStats(CurrentThreadState.CurrentGameFrame, Stack, nullptr, &NonFrameStats);
 
 	// cycle graphs
-	GenerateCycleGraphs( Stack, DataFrame.CycleGraphs );
+	GenerateCycleGraphs(Stack, DataFrame.CycleGraphs);
 
 	// accumulators
-	GenerateAccumulators( NonFrameStats, DataFrame.CountAccumulators, DataFrame.FloatAccumulators );
+	GenerateAccumulators(NonFrameStats, DataFrame.CountAccumulators, DataFrame.FloatAccumulators);
 }
 
+
 #endif
diff --git a/Engine/Source/Developer/ProfilerClient/Private/ProfilerClientManager.h b/Engine/Source/Developer/ProfilerClient/Private/ProfilerClientManager.h
--- a/Engine/Source/Developer/ProfilerClient/Private/ProfilerClientManager.h
+++ b/Engine/Source/Developer/ProfilerClient/Private/ProfilerClientManager.h
@@ -2,10 +2,20 @@
 
 #pragma once
 
+#include "IProfilerClient.h"
+#include "IProfilerServiceManager.h"
+#include "StatsData.h"
+#include "StatsFile.h"
+
+
 class FNewStatsReader;
 class FProfilerClientManager;
+class IMessageBus;
 
-/** Helper struct containing all of the data and operations associated with a service connection */
+
+/**
+ * Helper struct containing all of the data and operations associated with a service connection.
+ */
 struct FServiceConnection
 {
 	FServiceConnection();
@@ -27,7 +37,7 @@ struct FServiceConnection
 	FProfilerDataFrame CurrentData;
 
 	/** Initialize the connection from the given authorization message and context */
-	void Initialize( const FProfilerServiceAuthorize& Message, const IMessageContextRef& Context );
+	void Initialize(const FProfilerServiceAuthorize& Message, const TSharedRef<IMessageContext, ESPMode::ThreadSafe>& Context);
 
 #if STATS
 	/** Current stats data */
@@ -65,7 +75,7 @@ struct FServiceConnection
 	void GenerateAccumulators(TArray<FStatMessage>& Stats, TArray<FProfilerCountAccumulator>& CountAccumulators, TArray<FProfilerFloatAccumulator>& FloatAccumulators);
 
 	// adds a new style stat FName to the list of stats and generates an old style id and description
-	int32 FindOrAddStat( const FStatNameAndInfo& StatNameAndInfo, uint32 StatType );
+	int32 FindOrAddStat(const FStatNameAndInfo& StatNameAndInfo, uint32 StatType);
 
 	// adds a new style stat FName to the list of threads and generates an old style id and description
 	int32 FindOrAddThread(const FStatNameAndInfo& Thread);
@@ -78,11 +88,12 @@ struct FServiceConnection
 		Transition into new system
 	-----------------------------------------------------------------------------*/
 
-	void LoadCapture( const FString& DataFilepath, FProfilerClientManager* ProfilerClientManager );
+	void LoadCapture(const FString& DataFilepath, FProfilerClientManager* ProfilerClientManager);
 
 	FNewStatsReader* StatsReader;
 };
 
+
 /**
  * Implements the ProfileClient manager.
  */
@@ -98,97 +109,61 @@ public:
 	 *
 	 * @param InMessageBus The message bus to use.
 	 */
-	FProfilerClientManager( const IMessageBusRef& InMessageBus );
+	FProfilerClientManager(const TSharedRef<IMessageBus, ESPMode::ThreadSafe>& InMessageBus);
 
 	/** Default destructor */
 	~FProfilerClientManager();
 
 public:
 
-	// IProfilerClient Interface
+	//~ IProfilerClient Interface
 
-	virtual void Subscribe( const FGuid& Session ) override;
-	virtual void Track( const FGuid& Instance ) override;
-	virtual void Untrack( const FGuid& Instance ) override;
+	virtual void Subscribe(const FGuid& Session) override;
+	virtual void Track(const FGuid& Instance) override;
+	virtual void Untrack(const FGuid& Instance) override;
 	virtual void Unsubscribe() override;
-	virtual void SetCaptureState( const bool bRequestedCaptureState, const FGuid& InstanceId = FGuid() ) override;
-	virtual void SetPreviewState( const bool bRequestedPreviewState, const FGuid& InstanceId = FGuid() ) override;
-	virtual void LoadCapture( const FString& DataFilepath, const FGuid& ProfileId ) override;
-	virtual void RequestLastCapturedFile( const FGuid& InstanceId = FGuid() ) override;
-
-	virtual const FStatMetaData& GetStatMetaData( const FGuid& InstanceId ) const override
-	{
-		return Connections.Find(InstanceId)->StatMetaData;
-	}
-
-	virtual FProfilerClientDataDelegate& OnProfilerData() override
-	{
-		return ProfilerDataDelegate;
-	}
-
-	virtual FProfilerFileTransferDelegate& OnProfilerFileTransfer() override
-	{
-		return ProfilerFileTransferDelegate;
-	}
-
-	virtual FProfilerClientConnectedDelegate& OnProfilerClientConnected() override
-	{
-		return ProfilerClientConnectedDelegate;
-	}
-
-	virtual FProfilerClientDisconnectedDelegate& OnProfilerClientDisconnected() override
-	{
-		return ProfilerClientDisconnectedDelegate;
-	}
-
-	virtual FProfilerMetaDataUpdateDelegate& OnMetaDataUpdated() override
-	{
-		return ProfilerMetaDataUpdatedDelegate;
-	}
-
-	virtual FProfilerLoadStartedDelegate& OnLoadStarted() override
-	{
-		return ProfilerLoadStartedDelegate;
-	}
-
-	virtual FProfilerLoadCompletedDelegate& OnLoadCompleted() override
-	{
-		return ProfilerLoadCompletedDelegate;
-	}
-
-	virtual FProfilerLoadCancelledDelegate& OnLoadCancelled() override
-	{
-		return ProfilerLoadCancelledDelegate;
-	}
+	virtual void SetCaptureState(const bool bRequestedCaptureState, const FGuid& InstanceId = FGuid()) override;
+	virtual void SetPreviewState(const bool bRequestedPreviewState, const FGuid& InstanceId = FGuid()) override;
+	virtual void LoadCapture(const FString& DataFilepath, const FGuid& ProfileId) override;
+	virtual void RequestLastCapturedFile(const FGuid& InstanceId = FGuid()) override;
+	virtual const FStatMetaData& GetStatMetaData(const FGuid& InstanceId) const override;
+	virtual FProfilerClientDataDelegate& OnProfilerData() override;
+	virtual FProfilerFileTransferDelegate& OnProfilerFileTransfer() override;
+	virtual FProfilerClientConnectedDelegate& OnProfilerClientConnected() override;
+	virtual FProfilerClientDisconnectedDelegate& OnProfilerClientDisconnected() override;
+	virtual FProfilerMetaDataUpdateDelegate& OnMetaDataUpdated() override;
+	virtual FProfilerLoadStartedDelegate& OnLoadStarted() override;
+	virtual FProfilerLoadCompletedDelegate& OnLoadCompleted() override;
+	virtual FProfilerLoadCancelledDelegate& OnLoadCancelled() override;
 
 private:
 
-	// Handles message bus shutdowns.
+	/** Handles message bus shutdowns. */
 	void HandleMessageBusShutdown();
 
-	// Handles FProfilerServiceAuthorize2 messages.
-	void HandleServiceAuthorizeMessage( const FProfilerServiceAuthorize& Message, const IMessageContextRef& Context );
+	/** Handles FProfilerServiceAuthorize2 messages. */
+	void HandleServiceAuthorizeMessage(const FProfilerServiceAuthorize& Message, const TSharedRef<IMessageContext, ESPMode::ThreadSafe>& Context);
 
-	// Handles FProfilerServiceFileChunk messages.
-	void HandleServiceFileChunk( const FProfilerServiceFileChunk& FileChunk, const IMessageContextRef& Context );
+	/** Handles FProfilerServiceFileChunk messages. */
+	void HandleServiceFileChunk(const FProfilerServiceFileChunk& FileChunk, const TSharedRef<IMessageContext, ESPMode::ThreadSafe>& Context);
 
-	// Handles FProfilerServicePing messages.
-	void HandleServicePingMessage( const FProfilerServicePing& Message, const IMessageContextRef& Context );
+	/** Handles FProfilerServicePing messages. */
+	void HandleServicePingMessage(const FProfilerServicePing& Message, const TSharedRef<IMessageContext, ESPMode::ThreadSafe>& Context);
 
-	// Handles ticker callbacks (used to retry connection with a profiler service).
-	bool HandleTicker( float DeltaTime );
+	/** Handles ticker callbacks (used to retry connection with a profiler service). */
+	bool HandleTicker(float DeltaTime);
 
-	// Handles FProfilerServiceData2 messages.
-	void HandleProfilerServiceData2Message( const FProfilerServiceData2& Message, const IMessageContextRef& Context );
+	/** Handles FProfilerServiceData2 messages. */
+	void HandleProfilerServiceData2Message(const FProfilerServiceData2& Message, const TSharedRef<IMessageContext, ESPMode::ThreadSafe>& Context);
 
-	// Handles FProfilerServicePreviewAck messages
-	void HandleServicePreviewAckMessage( const FProfilerServicePreviewAck& Messsage, const IMessageContextRef& Context );
+	/** Handles FProfilerServicePreviewAck messages. */
+	void HandleServicePreviewAckMessage(const FProfilerServicePreviewAck& Messsage, const TSharedRef<IMessageContext, ESPMode::ThreadSafe>& Context);
 
-	// Handles ticker callbacks for sending out the received messages to subscribed clients
-	bool HandleMessagesTicker( float DeltaTime );
+	/** Handles ticker callbacks for sending out the received messages to subscribed clients. */
+	bool HandleMessagesTicker(float DeltaTime);
 
 	/** If hash is ok, writes the data to the archive and returns true, otherwise only returns false. */
-	bool CheckHashAndWrite( const FProfilerServiceFileChunk& FileChunk, const FProfilerFileChunkHeader& FileChunkHeader, FArchive* Writer );
+	bool CheckHashAndWrite(const FProfilerServiceFileChunk& FileChunk, const FProfilerFileChunkHeader& FileChunkHeader, FArchive* Writer);
 
 	/** Broadcast that loading has completed and cleans internal structures. */
 	void FinalizeLoading(const FGuid InstanceId);
@@ -197,13 +172,13 @@ private:
 	void CancelLoading(const FGuid InstanceId);
 
 	/** Decompress all stats data and send to the game thread. */
-	void DecompressDataAndSendToGame( FProfilerServiceData2* ToProcess );
+	void DecompressDataAndSendToGame(FProfilerServiceData2* ToProcess);
 
 	/** Sends decompressed data to the game thread. */
-	void SendDataToGame( TArray<uint8>* DataToGame, int64 Frame, const FGuid InstanceId );
+	void SendDataToGame(TArray<uint8>* DataToGame, int64 Frame, const FGuid InstanceId);
 
 	/** A new profiler data frame from the async loading thread. */
-	void SendProfilerDataFrameToGame( FProfilerDataFrame* NewData, FStatMetaData* MetaDataPtr, const FGuid InstanceId);
+	void SendProfilerDataFrameToGame(FProfilerDataFrame* NewData, FStatMetaData* MetaDataPtr, const FGuid InstanceId);
 
 	/** Removes active transfers and core tickers. */
 	void Shutdown();
@@ -223,11 +198,11 @@ private:
 
 	struct FReceivedFileInfo
 	{
-		FReceivedFileInfo( FArchive* InFileWriter, const int64 InProgress, const FString& InDestFilepath)
-			: FileWriter( InFileWriter )
-			, Progress( InProgress )
-			, DestFilepath( InDestFilepath )
-			, LastReceivedChunkTime( FPlatformTime::Seconds() )
+		FReceivedFileInfo(FArchive* InFileWriter, const int64 InProgress, const FString& InDestFilepath)
+			: FileWriter(InFileWriter)
+			, Progress(InProgress)
+			, DestFilepath(InDestFilepath)
+			, LastReceivedChunkTime(FPlatformTime::Seconds())
 		{}
 
 		bool IsTimedOut() const
@@ -257,7 +232,7 @@ private:
 	FMessageEndpointPtr MessageEndpoint;
 
 	/** Holds a pointer to the message bus. */
-	IMessageBusPtr MessageBus;
+	TSharedPtr<IMessageBus, ESPMode::ThreadSafe> MessageBus;
 
 	/** Delegate for notifying clients of received data */
 	FProfilerClientDataDelegate ProfilerDataDelegate;
@@ -304,7 +279,6 @@ private:
 	/** Handle to the registered OnShutdown for Message Bus. */
 	FDelegateHandle OnShutdownMessageBusDelegateHandle;
 
-
 	/** Holds the last time a ping was made to instances */
 	FDateTime LastPingTime;
 };
diff --git a/Engine/Source/Developer/ProfilerClient/Private/ProfilerClientModule.cpp b/Engine/Source/Developer/ProfilerClient/Private/ProfilerClientModule.cpp
--- a/Engine/Source/Developer/ProfilerClient/Private/ProfilerClientModule.cpp
+++ b/Engine/Source/Developer/ProfilerClient/Private/ProfilerClientModule.cpp
@@ -1,10 +1,11 @@
 // Copyright 1998-2016 Epic Games, Inc. All Rights Reserved.
 
-/*=============================================================================
-	ProfilerClientModule.cpp: Implements the FProfilerClientModule class.
-=============================================================================*/
-
 #include "ProfilerClientPrivatePCH.h"
+#include "IProfilerClientModule.h"
+#include "ProfilerClientManager.h"
+
+
+class IMessageBus;
 
 
 /**
@@ -15,7 +16,7 @@ class FProfilerClientModule
 {
 public:
 
-	// Begin IModuleInterface interface
+	//~ IModuleInterface interface
 
 	virtual void StartupModule() override
 	{
@@ -24,17 +25,16 @@ public:
 
 	virtual void ShutdownModule() override
 	{
+		// do nothing
 	}
 
-	// End IModuleInterface interface
-
 public:
 
-	// Begin IProfilerClientModule interface
+	//~ IProfilerClientModule interface
 
-	virtual IProfilerClientPtr CreateProfilerClient() override
+	virtual TSharedPtr<IProfilerClient> CreateProfilerClient() override
 	{
-		IMessageBusPtr MessageBus = MessageBusPtr.Pin();
+		auto MessageBus = MessageBusPtr.Pin();
 
 		if (!MessageBus.IsValid())
 		{
@@ -44,17 +44,11 @@ public:
 		return MakeShareable(new FProfilerClientManager(MessageBus.ToSharedRef()));
 	}
 
-	// End IProfilerClientModule interface
-
 private:
 
-	// Holds a weak pointer to the message bus.
-	IMessageBusWeakPtr MessageBusPtr;
+	/** Holds a weak pointer to the message bus. */
+	TWeakPtr<IMessageBus, ESPMode::ThreadSafe> MessageBusPtr;
 };
 
 
-/* Static initialization
- *****************************************************************************/
-
-
 IMPLEMENT_MODULE(FProfilerClientModule, ProfilerClient);
diff --git a/Engine/Source/Developer/ProfilerClient/Private/ProfilerClientPrivatePCH.h b/Engine/Source/Developer/ProfilerClient/Private/ProfilerClientPrivatePCH.h
--- a/Engine/Source/Developer/ProfilerClient/Private/ProfilerClientPrivatePCH.h
+++ b/Engine/Source/Developer/ProfilerClient/Private/ProfilerClientPrivatePCH.h
@@ -2,20 +2,15 @@
 
 #pragma once
 
-#include "ProfilerClient.h"
-
-
-/* Private dependencies
- *****************************************************************************/
-
-#include "Messaging.h"
-#include "ProfilerMessages.h"
-#include "Ticker.h"
-#include "StatsData.h"
-#include "StatsFile.h"
-
-
-/* Private includes
- *****************************************************************************/
-
-#include "ProfilerClientManager.h"
+#include "Runtime/Core/Public/Core.h"
+#include "Runtime/Core/Public/Containers/Ticker.h"
+#include "Runtime/Core/Public/Misc/SecureHash.h"
+#include "Runtime/Core/Public/Stats/StatsData.h"
+#include "Runtime/Core/Public/Stats/StatsFile.h"
+#include "Runtime/Messaging/Public/IMessageBus.h"
+#include "Runtime/Messaging/Public/IMessageContext.h"
+#include "Runtime/Messaging/Public/Helpers/MessageEndpointBuilder.h"
+#include "Runtime/SessionServices/Public/ISessionInstanceInfo.h"
+
+
+#include "Developer/ProfilerMessages/Public/ProfilerServiceMessages.h"
diff --git a/Engine/Source/Developer/ProfilerClient/Public/IProfilerClient.h b/Engine/Source/Developer/ProfilerClient/Public/IProfilerClient.h
new file mode 100644
--- /dev/null
+++ b/Engine/Source/Developer/ProfilerClient/Public/IProfilerClient.h
@@ -0,0 +1,169 @@
+// Copyright 1998-2016 Epic Games, Inc. All Rights Reserved.
+
+#pragma once
+
+
+struct FProfilerDataFrame;
+struct FStatMetaData;
+class ISessionInstanceInfo;
+
+
+/** Delegate for passing profiler data to UI */
+DECLARE_MULTICAST_DELEGATE_TwoParams(FProfilerClientDataDelegate, const FGuid&, const FProfilerDataFrame&);
+
+/** Delegate for alerting UI a session has been established */
+DECLARE_MULTICAST_DELEGATE_TwoParams(FProfilerClientConnectedDelegate, const FGuid&, const FGuid&);
+
+/** Delegate for alerting UI a session has been disconnected */
+DECLARE_MULTICAST_DELEGATE_TwoParams(FProfilerClientDisconnectedDelegate, const FGuid&, const FGuid&);
+
+/** Delegate for alerting subscribers the meta data has been updated */
+DECLARE_MULTICAST_DELEGATE_TwoParams(FProfilerMetaDataUpdateDelegate, const FGuid&, const FStatMetaData&);
+
+/** Delegate for alerting clients a load has started */
+DECLARE_MULTICAST_DELEGATE_OneParam(FProfilerLoadStartedDelegate, const FGuid&);
+
+/** Delegate for alerting clients a load has completed */
+DECLARE_MULTICAST_DELEGATE_OneParam(FProfilerLoadCompletedDelegate, const FGuid&);
+
+/** Delegate for alerting clients a load was cancelled */
+DECLARE_MULTICAST_DELEGATE_OneParam(FProfilerLoadCancelledDelegate, const FGuid&);
+
+/** Delegate for alerting clients a load has loaded meta data */
+DECLARE_MULTICAST_DELEGATE_OneParam(FProfilerLoadedMetaDataDelegate, const FGuid&);
+
+/** Delegate for updating the current progress of file transfer. */
+DECLARE_MULTICAST_DELEGATE_ThreeParams(FProfilerFileTransferDelegate, const FString& /*Filename*/, int64 /*FileProgress*/, int64 /*FileSize*/);
+
+
+/**
+ * Interface for Profiler Client.
+ */
+class IProfilerClient
+{
+public:
+
+	/**
+	 * Subscribe to the given profiler session id
+	 *
+	 * @param Session The session to subscribe to for profiler data
+	 */
+	virtual void Subscribe(const FGuid& Session) = 0;
+
+	/**
+	 * 
+	 */
+	virtual void Track(const FGuid& Instance) = 0;
+
+	/**
+	 * 
+	 */
+	virtual void Untrack(const FGuid& Instance) = 0;
+
+	/** Unsubscribe from all sessions. */
+	virtual void Unsubscribe() = 0;
+
+	/**
+	 * Changes the current state of the capturing data service-side.
+	 *
+	 * @param bRequestedCaptureState Data capture state that should be set.
+	 * @param InstanceId If valid, this function will be executed only for the specified instance, otherwise will be executed on all instances.
+	 *
+	 */
+	virtual void SetCaptureState(const bool bRequestedCaptureState, const FGuid& InstanceId = FGuid()) = 0; 
+
+	/**
+	 * Changes the current state of the previewing capture data.
+	 *
+	 * @param bRequestedCaptureState Data preview state that should be set.
+	 * @param InstanceId If valid, this function will be executed only for the specified instance, otherwise will be executed on all instances.
+	 *
+	 */
+	virtual void SetPreviewState(const bool bRequestedPreviewState, const FGuid& InstanceId = FGuid()) = 0;
+
+	/**
+	 * Loads a Capture file
+	 */
+	virtual void LoadCapture(const FString& DataFilepath, const FGuid& ProfileId) = 0;
+
+	/**
+	 * Cancels a capture file load that is in progress
+	 */
+	virtual void CancelLoading(const FGuid InstanceId) = 0;
+
+	/**
+	 * Requests the last captured file from the service.
+	 *
+	 * @param InstanceId If valid, this function will be executed only for the specified instance, otherwise will be executed on all instances
+	 */
+	virtual void RequestLastCapturedFile(const FGuid& InstanceId = FGuid()) = 0; 
+
+	/**
+	 * Gets the description for the given stat id
+	 *
+	 * @param StatId The identifier of the statistic description to retrieve.
+	 * @return the FStatMetaData struct with the description.
+	 */
+	virtual const FStatMetaData& GetStatMetaData(const FGuid& InstanceId) const = 0;
+
+	/**
+	 * Retrieves the profiler data delegate.
+	 *
+	 * @return profiler data delegate.
+	 */
+	virtual FProfilerClientDataDelegate& OnProfilerData() = 0;
+
+	/**
+	 * Retrieves the profiler file transfer delegate.
+	 *
+	 * @return profiler file transfer delegate.
+	 */
+	virtual FProfilerFileTransferDelegate& OnProfilerFileTransfer() = 0;
+
+	/**
+	 * Retrieves the profiler client connected delegate.
+	 *
+	 * @return profiler client connected delegate.
+	 */
+	virtual FProfilerClientConnectedDelegate& OnProfilerClientConnected() = 0;
+
+	/**
+	 * Retrieves the profiler client disconnected delegate.
+	 *
+	 * @return profiler client disconnected delegate.
+	 */
+	virtual FProfilerClientDisconnectedDelegate& OnProfilerClientDisconnected() = 0;
+
+	/**
+	 * Retrieves the profiler meta data update delegate.
+	 *
+	 * @return profiler meta data delegate.
+	 */
+	virtual FProfilerMetaDataUpdateDelegate& OnMetaDataUpdated() = 0;
+
+	/**
+	 * Retrieves the load started delegate.
+	 *
+	 * @return profiler load start delegate.
+	 */
+	virtual FProfilerLoadStartedDelegate& OnLoadStarted() = 0;
+
+	/**
+	 * Retrieves the load completed delegate.
+	 *
+	 * @return profiler load completed delegate.
+	 */
+	virtual FProfilerLoadCompletedDelegate& OnLoadCompleted() = 0;
+
+	/**
+	* Retrieves the load cancelled delegate.
+	*
+	* @return profiler load cancelled delegate.
+	*/
+	virtual FProfilerLoadCancelledDelegate& OnLoadCancelled() = 0;
+
+protected:
+
+	/** Virtual destructor. */
+	virtual ~IProfilerClient() { }
+};
diff --git a/Engine/Source/Developer/ProfilerClient/Public/IProfilerClientModule.h b/Engine/Source/Developer/ProfilerClient/Public/IProfilerClientModule.h
new file mode 100644
--- /dev/null
+++ b/Engine/Source/Developer/ProfilerClient/Public/IProfilerClientModule.h
@@ -0,0 +1,28 @@
+// Copyright 1998-2016 Epic Games, Inc. All Rights Reserved.
+
+#pragma once
+
+
+class IProfilerClient;
+
+
+/**
+ * Interface for ProfilerClient modules.
+ */
+class IProfilerClientModule
+	: public IModuleInterface
+{
+public:
+
+	/**
+	 * Creates a Profiler client.
+	 *
+	 * @return The new client.
+	 */
+	virtual TSharedPtr<IProfilerClient> CreateProfilerClient() = 0;
+
+protected:
+
+	/** Virtual destructor */
+	virtual ~IProfilerClientModule() { }
+};
diff --git a/Engine/Source/Developer/ProfilerMessages/Private/ProfilerMessagesModule.cpp b/Engine/Source/Developer/ProfilerMessages/Private/ProfilerMessagesModule.cpp
--- a/Engine/Source/Developer/ProfilerMessages/Private/ProfilerMessagesModule.cpp
+++ b/Engine/Source/Developer/ProfilerMessages/Private/ProfilerMessagesModule.cpp
@@ -1,6 +1,7 @@
 // Copyright 1998-2016 Epic Games, Inc. All Rights Reserved.
 
 #include "ProfilerMessagesPrivatePCH.h"
+#include "ModuleInterface.h"
 
 
 /**
@@ -11,11 +12,10 @@ class FProfilerMessagesModule
 {
 public:
 
-	virtual void StartupModule( ) override { }
+	virtual void StartupModule() override { }
+	virtual void ShutdownModule() override { }
 
-	virtual void ShutdownModule( ) override { }
-
-	virtual bool SupportsDynamicReloading( ) override
+	virtual bool SupportsDynamicReloading() override
 	{
 		return true;
 	}
diff --git a/Engine/Source/Developer/ProfilerMessages/Private/ProfilerMessagesPrivatePCH.h b/Engine/Source/Developer/ProfilerMessages/Private/ProfilerMessagesPrivatePCH.h
--- a/Engine/Source/Developer/ProfilerMessages/Private/ProfilerMessagesPrivatePCH.h
+++ b/Engine/Source/Developer/ProfilerMessages/Private/ProfilerMessagesPrivatePCH.h
@@ -2,10 +2,5 @@
 
 #pragma once
 
-#include "ProfilerMessages.h"
-
-
-/* Private dependencies
- *****************************************************************************/
-
-#include "ModuleManager.h"
+#include "Runtime/Core/Public/Core.h"
+#include "Runtime/Core/Public/Modules/ModuleManager.h"
diff --git a/Engine/Source/Developer/ProfilerMessages/Public/ProfilerServiceMessages.h b/Engine/Source/Developer/ProfilerMessages/Public/ProfilerServiceMessages.h
new file mode 100644
--- /dev/null
+++ b/Engine/Source/Developer/ProfilerMessages/Public/ProfilerServiceMessages.h
@@ -0,0 +1,285 @@
+// Copyright 1998-2016 Epic Games, Inc. All Rights Reserved.
+
+#pragma once
+
+#include "ProfilerServiceMessages.generated.h"
+
+
+/** Profiler Service authorization message. */
+USTRUCT()
+struct FProfilerServiceAuthorize
+{
+	GENERATED_USTRUCT_BODY()
+
+	/** Session ID. */
+	UPROPERTY(EditAnywhere, Category="Message")
+	FGuid SessionId;
+
+	/** Instance ID. */
+	UPROPERTY(EditAnywhere, Category="Message")
+	FGuid InstanceId;
+
+	/** Default constructor. */
+	FProfilerServiceAuthorize() { }
+
+	/** Creates and initializes a new instance. */
+	FProfilerServiceAuthorize(const FGuid& InSessionId, const FGuid& InInstanceId)
+		: SessionId(InSessionId)
+		, InstanceId(InInstanceId)
+	{}
+
+};
+
+
+/** Profiler Service data. */
+USTRUCT()
+struct FProfilerServiceData2
+{
+	GENERATED_USTRUCT_BODY()
+
+	/** Instance ID. */
+	UPROPERTY(EditAnywhere, Category="Message")
+	FGuid InstanceId;
+
+	/** Stats frame. */
+	UPROPERTY(EditAnywhere, Category="Message")
+	int64 Frame;
+
+	/** Size of the compressed data. */
+	UPROPERTY(EditAnywhere, Category="Message")
+	int32 CompressedSize;
+	
+	/** Size of the uncompressed data. */
+	UPROPERTY(EditAnywhere, Category="Message")
+	int32 UncompressedSize;
+
+	/** Profiler data encoded as string of hexes, cannot use TArray<uint8> because of the Message Bus limitation. */
+	UPROPERTY(EditAnywhere, Category="Message")
+	FString HexData;
+
+	/** Default constructor. */
+	FProfilerServiceData2() { }
+
+	/** Creates and initializes a new instance. */
+	FProfilerServiceData2(const FGuid& InInstance, int64 InFrame, const FString& InHexData, int32 InCompressedSize, int32 InUncompressedSize)
+		: InstanceId(InInstance)
+		, Frame(InFrame)
+		, CompressedSize(InCompressedSize)
+		, UncompressedSize(InUncompressedSize)
+		, HexData(InHexData)
+	{ }
+};
+
+
+/**
+ */
+USTRUCT()
+struct FProfilerServicePreviewAck
+{
+	GENERATED_USTRUCT_BODY()
+
+	/** */
+	UPROPERTY(EditAnywhere, Category="Message")
+	FGuid InstanceId;
+
+	/** Default constructor. */
+	FProfilerServicePreviewAck() { }
+
+	/** Creates and initializes a new instance. */
+	FProfilerServicePreviewAck(const FGuid& InInstance)
+		: InstanceId(InInstance)
+	{ }
+};
+
+
+/**
+ * Implements a message for copying a file through the network, as well as for synchronization.
+ * Unfortunately assumes that InstanceId and Filename are transfered without errors.
+ */
+USTRUCT()
+struct FProfilerServiceFileChunk
+{
+	GENERATED_USTRUCT_BODY()
+
+	/** The ID of the instance where this message should be sent. */
+	UPROPERTY(EditAnywhere, Category="Message")
+	FGuid InstanceId;
+
+	/** The file containing this file chunk. */
+	UPROPERTY(EditAnywhere, Category="Message")
+	FString Filename;
+
+	/** Data to be sent through message bus. Message bug doesn't support TArray<>, so we encode the data as HexString. */
+	UPROPERTY(EditAnywhere, Category="Message")
+	FString HexData;
+
+	/** FProfilerFileChunkHeader stored in the array. */
+	UPROPERTY(EditAnywhere, Category="Message")
+	TArray<uint8> Header;
+
+	/** Hash of this data and header. */
+	UPROPERTY(EditAnywhere, Category="Message")
+	TArray<uint8> ChunkHash;
+
+	/** Default constructor. */
+	FProfilerServiceFileChunk() { }
+
+	/** Constructor for the new file chunk. */
+	FProfilerServiceFileChunk
+	( 
+		const FGuid& InInstanceID, 
+		const FString& InFilename, 
+		const TArray<uint8>& InHeader
+	)
+		: InstanceId(InInstanceID)
+		, Filename(InFilename)
+		, Header(InHeader)
+	{ }
+
+	struct FNullTag
+	{
+	};
+
+	/** Copy constructor, copies all properties, but not data. */
+	FProfilerServiceFileChunk(const FProfilerServiceFileChunk& ProfilerServiceFileChunk, FNullTag)
+		: InstanceId(ProfilerServiceFileChunk.InstanceId)
+		, Filename(ProfilerServiceFileChunk.Filename)
+		, Header(ProfilerServiceFileChunk.Header)
+	{ }
+};
+
+
+/**
+ */
+USTRUCT()
+struct FProfilerServicePing
+{
+	GENERATED_USTRUCT_BODY()
+};
+
+
+/**
+ */
+USTRUCT()
+struct FProfilerServicePong
+{
+	GENERATED_USTRUCT_BODY()
+
+};
+
+
+/**
+ */
+USTRUCT()
+struct FProfilerServiceSubscribe
+{
+	GENERATED_USTRUCT_BODY()
+
+	/** */
+	UPROPERTY(EditAnywhere, Category="Message")
+	FGuid SessionId;
+
+	/** */
+	UPROPERTY(EditAnywhere, Category="Message")
+	FGuid InstanceId;
+
+	/** Default constructor. */
+	FProfilerServiceSubscribe() { }
+
+	/** Creates and initializes a new instance. */
+	FProfilerServiceSubscribe(const FGuid& InSessionId, const FGuid& InInstanceId)
+		: SessionId(InSessionId)
+		, InstanceId(InInstanceId)
+	{ }
+};
+
+
+/**
+ */
+USTRUCT()
+struct FProfilerServiceUnsubscribe
+{
+	GENERATED_USTRUCT_BODY()
+
+	/** */
+	UPROPERTY(EditAnywhere, Category="Message")
+	FGuid SessionId;
+
+	/** */
+	UPROPERTY(EditAnywhere, Category="Message")
+	FGuid InstanceId;
+
+	/** Default constructor. */
+	FProfilerServiceUnsubscribe() { }
+
+	/** Creates and initializes a new instance. */
+	FProfilerServiceUnsubscribe(const FGuid& InSessionId, const FGuid& InInstanceId)
+		: SessionId(InSessionId)
+		, InstanceId(InInstanceId)
+	{ }
+};
+
+
+/**
+ */
+USTRUCT()
+struct FProfilerServiceCapture
+{
+	GENERATED_USTRUCT_BODY()
+
+	/** The data capture state that should be set. */
+	UPROPERTY(EditAnywhere, Category="Message")
+	bool bRequestedCaptureState;
+
+	/** Default constructor. */
+	FProfilerServiceCapture() { }
+
+	/** Creates and initializes a new instance. */
+	FProfilerServiceCapture(const bool bInRequestedCaptureState)
+		: bRequestedCaptureState(bInRequestedCaptureState)
+	{ }
+};
+
+
+/**
+ */
+USTRUCT()
+struct FProfilerServicePreview
+{
+	GENERATED_USTRUCT_BODY()
+
+	/**
+	 * The data preview state that should be set.
+	 */
+	UPROPERTY(EditAnywhere, Category="Message")
+	bool bRequestedPreviewState;
+
+	/** Default constructor. */
+	FProfilerServicePreview() { }
+
+	/** Creates and initializes a new instance. */
+	FProfilerServicePreview(const bool bInRequestedPreviewState)
+		: bRequestedPreviewState(bInRequestedPreviewState)
+	{ }
+};
+
+
+/**
+ */
+USTRUCT()
+struct FProfilerServiceRequest
+{
+	GENERATED_USTRUCT_BODY()
+
+	/** Request @see EProfilerRequestType. */
+	UPROPERTY(EditAnywhere, Category="Message")
+	uint32 Request;
+
+	/** Default constructor. */
+	FProfilerServiceRequest() { }
+
+	/** Creates and initializes a new instance. */
+	FProfilerServiceRequest(uint32 InRequest)
+		: Request(InRequest)
+	{ }
+};
diff --git a/Engine/Source/Developer/ProfilerService/Private/ProfilerServiceFileTransfer.cpp b/Engine/Source/Developer/ProfilerService/Private/ProfilerServiceFileTransfer.cpp
--- a/Engine/Source/Developer/ProfilerService/Private/ProfilerServiceFileTransfer.cpp
+++ b/Engine/Source/Developer/ProfilerService/Private/ProfilerServiceFileTransfer.cpp
@@ -1,18 +1,19 @@
 // Copyright 1998-2016 Epic Games, Inc. All Rights Reserved.
 
 #include "ProfilerServicePrivatePCH.h"
-#include "SecureHash.h"
+#include "ProfilerServiceManager.h"
 
 
-FFileTransferRunnable::FFileTransferRunnable( FMessageEndpointPtr& InMessageEndpoint )
-	: Runnable( nullptr )
-	, WorkEvent( FPlatformProcess::GetSynchEventFromPool( true ) )
-	, MessageEndpoint( InMessageEndpoint )
-	, StopTaskCounter( 0 )
+FFileTransferRunnable::FFileTransferRunnable(FMessageEndpointPtr& InMessageEndpoint)
+	: Runnable(nullptr)
+	, WorkEvent(FPlatformProcess::GetSynchEventFromPool(true))
+	, MessageEndpoint(InMessageEndpoint)
+	, StopTaskCounter(0)
 {
-	Runnable = FRunnableThread::Create( this, TEXT( "FFileTransferRunnable" ), 128 * 1024, TPri_BelowNormal );
+	Runnable = FRunnableThread::Create(this, TEXT("FFileTransferRunnable"), 128 * 1024, TPri_BelowNormal);
 }
 
+
 FFileTransferRunnable::~FFileTransferRunnable()
 {
 	if (Runnable != nullptr)
@@ -27,30 +28,32 @@ FFileTransferRunnable::~FFileTransferRunnable()
 	for (auto It = ActiveTransfers.CreateIterator(); It; ++It)
 	{
 		FReaderAndAddress& ReaderAndAddress = It.Value();
-		DeleteFileReader( ReaderAndAddress );
+		DeleteFileReader(ReaderAndAddress);
 
-		UE_LOG( LogProfilerService, Log, TEXT( "File service-client sending aborted (srv): %s" ), *It.Key() );
+		UE_LOG(LogProfilerService, Log, TEXT("File service-client sending aborted (srv): %s"), *It.Key());
 	}
 
-	FPlatformProcess::ReturnSynchEventToPool( WorkEvent );
+	FPlatformProcess::ReturnSynchEventToPool(WorkEvent);
 	WorkEvent = nullptr;
 }
 
+
 bool FFileTransferRunnable::Init()
 {
 	return true;
 }
 
+
 uint32 FFileTransferRunnable::Run()
 {
 	while (!ShouldStop())
 	{
-		if (WorkEvent->Wait( 250 ))
+		if (WorkEvent->Wait(250))
 		{
 			FProfilerServiceFileChunk* FileChunk;
-			while (!ShouldStop() && SendQueue.Dequeue( FileChunk ))
+			while (!ShouldStop() && SendQueue.Dequeue(FileChunk))
 			{
-				FMemoryReader MemoryReader( FileChunk->Header );
+				FMemoryReader MemoryReader(FileChunk->Header);
 				FProfilerFileChunkHeader FileChunkHeader;
 				MemoryReader << FileChunkHeader;
 				FileChunkHeader.Validate();
@@ -61,8 +64,8 @@ uint32 FFileTransferRunnable::Run()
 					FArchive* ReaderArchive = nullptr;
 					FMessageAddress Recipient;
 					{
-						FScopeLock Lock( &SyncActiveTransfers );
-						const FReaderAndAddress* ReaderAndAddress = ActiveTransfers.Find( FileChunk->Filename );
+						FScopeLock Lock(&SyncActiveTransfers);
+						const FReaderAndAddress* ReaderAndAddress = ActiveTransfers.Find(FileChunk->Filename);
 						if (ReaderAndAddress)
 						{
 							ReaderArchive = ReaderAndAddress->Key;
@@ -73,17 +76,17 @@ uint32 FFileTransferRunnable::Run()
 					// If there is no reader and recipient is invalid, it means that the file transfer is no longer valid, because client disconnected or exited.
 					if (ReaderArchive && Recipient.IsValid())
 					{
-						ReadAndSetHash( FileChunk, FileChunkHeader, ReaderArchive );
+						ReadAndSetHash(FileChunk, FileChunkHeader, ReaderArchive);
 
 						if (MessageEndpoint.IsValid())
 						{
-							MessageEndpoint->Send( FileChunk, Recipient );
+							MessageEndpoint->Send(FileChunk, Recipient);
 						}
 					}
 				}
 				else if (FileChunkHeader.ChunkType == EProfilerFileChunkType::PrepareFile)
 				{
-					PrepareFileForSending( FileChunk );
+					PrepareFileForSending(FileChunk);
 				}
 			}
 
@@ -94,42 +97,45 @@ uint32 FFileTransferRunnable::Run()
 	return 0;
 }
 
+
 void FFileTransferRunnable::Exit()
 {
 
 }
 
-void FFileTransferRunnable::EnqueueFileToSend( const FString& StatFilename, const FMessageAddress& RecipientAddress, const FGuid& ServiceInstanceId )
+
+void FFileTransferRunnable::EnqueueFileToSend(const FString& StatFilename, const FMessageAddress& RecipientAddress, const FGuid& ServiceInstanceId)
 {
-	UE_LOG( LogProfilerService, Log, TEXT( "Opening stats file for service-client sending: %s" ), *StatFilename );
+	UE_LOG(LogProfilerService, Log, TEXT("Opening stats file for service-client sending: %s"), *StatFilename);
 
-	const int64 FileSize = IFileManager::Get().FileSize( *StatFilename );
+	const int64 FileSize = IFileManager::Get().FileSize(*StatFilename);
 	if (FileSize < 4)
 	{
-		UE_LOG( LogProfilerService, Error, TEXT( "Could not open: %s" ), *StatFilename );
+		UE_LOG(LogProfilerService, Error, TEXT("Could not open: %s"), *StatFilename);
 		return;
 	}
 
-	FArchive* FileReader = IFileManager::Get().CreateFileReader( *StatFilename );
+	FArchive* FileReader = IFileManager::Get().CreateFileReader(*StatFilename);
 	if (!FileReader)
 	{
-		UE_LOG( LogProfilerService, Error, TEXT( "Could not open: %s" ), *StatFilename );
+		UE_LOG(LogProfilerService, Error, TEXT("Could not open: %s"), *StatFilename);
 		return;
 	}
 
 	{
-		FScopeLock Lock( &SyncActiveTransfers );
-		check( !ActiveTransfers.Contains( StatFilename ) );
-		ActiveTransfers.Add( StatFilename, FReaderAndAddress( FileReader, RecipientAddress ) );
+		FScopeLock Lock(&SyncActiveTransfers);
+		check(!ActiveTransfers.Contains(StatFilename));
+		ActiveTransfers.Add(StatFilename, FReaderAndAddress(FileReader, RecipientAddress));
 	}
 
 	// This is not a real file chunk, but helper to prepare file for sending on the runnable.
-	EnqueueFileChunkToSend( new FProfilerServiceFileChunk( ServiceInstanceId, StatFilename, FProfilerFileChunkHeader( 0, 0, FileReader->TotalSize(), EProfilerFileChunkType::PrepareFile ).AsArray() ), true );
+	EnqueueFileChunkToSend(new FProfilerServiceFileChunk(ServiceInstanceId, StatFilename, FProfilerFileChunkHeader(0, 0, FileReader->TotalSize(), EProfilerFileChunkType::PrepareFile).AsArray()), true);
 }
 
-void FFileTransferRunnable::EnqueueFileChunkToSend( FProfilerServiceFileChunk* FileChunk, bool bTriggerWorkEvent /*= false */ )
+
+void FFileTransferRunnable::EnqueueFileChunkToSend(FProfilerServiceFileChunk* FileChunk, bool bTriggerWorkEvent /*= false */)
 {
-	SendQueue.Enqueue( FileChunk );
+	SendQueue.Enqueue(FileChunk);
 
 	if (bTriggerWorkEvent)
 	{
@@ -138,14 +144,15 @@ void FFileTransferRunnable::EnqueueFileChunkToSend( FProfilerServiceFileChunk* F
 	}
 }
 
-void FFileTransferRunnable::PrepareFileForSending( FProfilerServiceFileChunk*& FileChunk )
+
+void FFileTransferRunnable::PrepareFileForSending(FProfilerServiceFileChunk*& FileChunk)
 {
 	// Find the corresponding file archive and recipient.
 	FArchive* Reader = nullptr;
 	FMessageAddress Recipient;
 	{
-		FScopeLock Lock( &SyncActiveTransfers );
-		const FReaderAndAddress& ReaderAndAddress = ActiveTransfers.FindChecked( FileChunk->Filename );
+		FScopeLock Lock(&SyncActiveTransfers);
+		const FReaderAndAddress& ReaderAndAddress = ActiveTransfers.FindChecked(FileChunk->Filename);
 		Reader = ReaderAndAddress.Key;
 		Recipient = ReaderAndAddress.Value;
 	}
@@ -155,9 +162,9 @@ void FFileTransferRunnable::PrepareFileForSending( FProfilerServiceFileChunk*& F
 
 	while (RemainingSizeToSend > 0)
 	{
-		const int64 SizeToCopy = FMath::Min( FProfilerFileChunkHeader::DefChunkSize, RemainingSizeToSend );
+		const int64 SizeToCopy = FMath::Min(FProfilerFileChunkHeader::DefChunkSize, RemainingSizeToSend);
 
-		EnqueueFileChunkToSend( new FProfilerServiceFileChunk( FileChunk->InstanceId, FileChunk->Filename, FProfilerFileChunkHeader( ChunkOffset, SizeToCopy, Reader->TotalSize(), EProfilerFileChunkType::SendChunk ).AsArray() ) );
+		EnqueueFileChunkToSend(new FProfilerServiceFileChunk(FileChunk->InstanceId, FileChunk->Filename, FProfilerFileChunkHeader(ChunkOffset, SizeToCopy, Reader->TotalSize(), EProfilerFileChunkType::SendChunk).AsArray()));
 
 		ChunkOffset += SizeToCopy;
 		RemainingSizeToSend -= SizeToCopy;
@@ -171,61 +178,65 @@ void FFileTransferRunnable::PrepareFileForSending( FProfilerServiceFileChunk*& F
 	FileChunk = nullptr;
 }
 
-void FFileTransferRunnable::FinalizeFileSending( const FString& Filename )
+
+void FFileTransferRunnable::FinalizeFileSending(const FString& Filename)
 {
-	FScopeLock Lock( &SyncActiveTransfers );
+	FScopeLock Lock(&SyncActiveTransfers);
 
-	check( ActiveTransfers.Contains( Filename ) );
-	FReaderAndAddress ReaderAndAddress = ActiveTransfers.FindAndRemoveChecked( Filename );
-	DeleteFileReader( ReaderAndAddress );
+	check(ActiveTransfers.Contains(Filename));
+	FReaderAndAddress ReaderAndAddress = ActiveTransfers.FindAndRemoveChecked(Filename);
+	DeleteFileReader(ReaderAndAddress);
 
-	UE_LOG( LogProfilerService, Log, TEXT( "File service-client sent successfully : %s" ), *Filename );
+	UE_LOG(LogProfilerService, Log, TEXT("File service-client sent successfully : %s"), *Filename);
 }
 
-void FFileTransferRunnable::AbortFileSending( const FMessageAddress& Recipient )
+
+void FFileTransferRunnable::AbortFileSending(const FMessageAddress& Recipient)
 {
-	FScopeLock Lock( &SyncActiveTransfers );
+	FScopeLock Lock(&SyncActiveTransfers);
 
 	for (auto It = ActiveTransfers.CreateIterator(); It; ++It)
 	{
 		FReaderAndAddress& ReaderAndAddress = It.Value();
 		if (ReaderAndAddress.Value == Recipient)
 		{
-			UE_LOG( LogProfilerService, Log, TEXT( "File service-client sending aborted (cl): %s" ), *It.Key() );
-			FReaderAndAddress ActiveReaderAndAddress = ActiveTransfers.FindAndRemoveChecked( It.Key() );
-			DeleteFileReader( ActiveReaderAndAddress );
+			UE_LOG(LogProfilerService, Log, TEXT("File service-client sending aborted (cl): %s"), *It.Key());
+			FReaderAndAddress ActiveReaderAndAddress = ActiveTransfers.FindAndRemoveChecked(It.Key());
+			DeleteFileReader(ActiveReaderAndAddress);
 		}
 	}
 }
 
-void FFileTransferRunnable::DeleteFileReader( FReaderAndAddress& ReaderAndAddress )
+
+void FFileTransferRunnable::DeleteFileReader(FReaderAndAddress& ReaderAndAddress)
 {
 	delete ReaderAndAddress.Key;
 	ReaderAndAddress.Key = nullptr;
 }
 
-void FFileTransferRunnable::ReadAndSetHash( FProfilerServiceFileChunk* FileChunk, const FProfilerFileChunkHeader& FileChunkHeader, FArchive* Reader )
+
+void FFileTransferRunnable::ReadAndSetHash(FProfilerServiceFileChunk* FileChunk, const FProfilerFileChunkHeader& FileChunkHeader, FArchive* Reader)
 {
 	TArray<uint8> FileChunkData;
 
-	FileChunkData.AddUninitialized( FileChunkHeader.ChunkSize );
+	FileChunkData.AddUninitialized(FileChunkHeader.ChunkSize);
 
-	Reader->Seek( FileChunkHeader.ChunkOffset );
-	Reader->Serialize( FileChunkData.GetData(), FileChunkHeader.ChunkSize );
+	Reader->Seek(FileChunkHeader.ChunkOffset);
+	Reader->Serialize(FileChunkData.GetData(), FileChunkHeader.ChunkSize);
 
 	const int32 HashSize = 20;
 	uint8 LocalHash[HashSize] = {0};
 
 	// Hash file chunk data. 
 	FSHA1 Sha;
-	Sha.Update( FileChunkData.GetData(), FileChunkHeader.ChunkSize );
+	Sha.Update(FileChunkData.GetData(), FileChunkHeader.ChunkSize);
 	// Hash file chunk header.
-	Sha.Update( FileChunk->Header.GetData(), FileChunk->Header.Num() );
+	Sha.Update(FileChunk->Header.GetData(), FileChunk->Header.Num());
 	Sha.Final();
-	Sha.GetHash( LocalHash );
+	Sha.GetHash(LocalHash);
 
-	FileChunk->ChunkHash.AddUninitialized( HashSize );
-	FMemory::Memcpy( FileChunk->ChunkHash.GetData(), LocalHash, HashSize );
+	FileChunk->ChunkHash.AddUninitialized(HashSize);
+	FMemory::Memcpy(FileChunk->ChunkHash.GetData(), LocalHash, HashSize);
 
 	// Limit transfer per second, otherwise we will probably hang the message bus.
 	static int64 TotalReadBytes = 0;
@@ -236,12 +247,12 @@ void FFileTransferRunnable::ReadAndSetHash( FProfilerServiceFileChunk* FileChunk
 #endif // _DEBUG
 
 	// Convert to hex.
-	FileChunk->HexData = FString::FromHexBlob( FileChunkData.GetData(), FileChunkData.Num() );
+	FileChunk->HexData = FString::FromHexBlob(FileChunkData.GetData(), FileChunkData.Num());
 
 	TotalReadBytes += FileChunkHeader.ChunkSize;
 	if (TotalReadBytes > NumBytesPerTick)
 	{
-		FPlatformProcess::Sleep( 0.1f );
+		FPlatformProcess::Sleep(0.1f);
 		TotalReadBytes = 0;
 	}
 }
diff --git a/Engine/Source/Developer/ProfilerService/Private/ProfilerServiceManager.cpp b/Engine/Source/Developer/ProfilerService/Private/ProfilerServiceManager.cpp
--- a/Engine/Source/Developer/ProfilerService/Private/ProfilerServiceManager.cpp
+++ b/Engine/Source/Developer/ProfilerService/Private/ProfilerServiceManager.cpp
@@ -1,16 +1,14 @@
 // Copyright 1998-2016 Epic Games, Inc. All Rights Reserved.
 
 #include "ProfilerServicePrivatePCH.h"
+#include "ProfilerServiceManager.h"
 
-#include "StatsData.h"
-#include "StatsFile.h"
 
+DEFINE_LOG_CATEGORY(LogProfilerService);
 
-DEFINE_LOG_CATEGORY( LogProfilerService );
 
-/*-----------------------------------------------------------------------------
-	FProfilerServiceStatsStream
------------------------------------------------------------------------------*/
+/* FProfilerServiceStatsStream
+ *****************************************************************************/
 
 #if STATS
 
@@ -37,30 +35,7 @@ public:
 	}
 };
 
-#endif
-
-/*-----------------------------------------------------------------------------
-	FProfilerServiceManager
------------------------------------------------------------------------------*/
-
-void FProfilerServiceManager::HandleServiceFileChunkMessage( const FProfilerServiceFileChunk& Message, const IMessageContextRef& Context )
-{
-	FMemoryReader Reader( Message.Header );
-	FProfilerFileChunkHeader Header;
-	Reader << Header;
-	Header.Validate();
-
-	if (Header.ChunkType == EProfilerFileChunkType::SendChunk)
-	{
-		// Send this file chunk again.
-		FileTransferRunnable->EnqueueFileChunkToSend( new FProfilerServiceFileChunk( Message, FProfilerServiceFileChunk::FNullTag() ), true );
-	}
-	else if (Header.ChunkType == EProfilerFileChunkType::FinalizeFile)
-	{
-		// Finalize file.
-		FileTransferRunnable->FinalizeFileSending( Message.Filename );
-	}
-}
+#endif //STATS
 
 
 /* FProfilerServiceManager structors
@@ -73,6 +48,7 @@ FProfilerServiceManager::FProfilerServiceManager()
 	PingDelegate = FTickerDelegate::CreateRaw(this, &FProfilerServiceManager::HandlePing);
 }
 
+
 /* IProfilerServiceManager interface
  *****************************************************************************/
 
@@ -93,6 +69,7 @@ void FProfilerServiceManager::StopCapture()
 #endif
 }
 
+
 /* FProfilerServiceManager implementation
  *****************************************************************************/
 
@@ -131,9 +108,9 @@ void FProfilerServiceManager::Shutdown()
 }
 
 
-IProfilerServiceManagerPtr FProfilerServiceManager::CreateSharedServiceManager()
+TSharedPtr<IProfilerServiceManager> FProfilerServiceManager::CreateSharedServiceManager()
 {
-	static IProfilerServiceManagerPtr ProfilerServiceManager;
+	static TSharedPtr<IProfilerServiceManager> ProfilerServiceManager;
 
 	if (!ProfilerServiceManager.IsValid())
 	{
@@ -143,6 +120,7 @@ IProfilerServiceManagerPtr FProfilerServiceManager::CreateSharedServiceManager()
 	return ProfilerServiceManager;
 }
 
+
 void FProfilerServiceManager::AddNewFrameHandleStatsThread()
 {
 #if	STATS
@@ -150,9 +128,10 @@ void FProfilerServiceManager::AddNewFrameHandleStatsThread()
 	NewFrameDelegateHandle = Stats.NewFrameDelegate.AddRaw( this, &FProfilerServiceManager::HandleNewFrame );
 	StatsMasterEnableAdd();
 	MetadataSize = 0;
-#endif // STATS
+#endif //STATS
 }
 
+
 void FProfilerServiceManager::RemoveNewFrameHandleStatsThread()
 {
 #if	STATS
@@ -160,9 +139,10 @@ void FProfilerServiceManager::RemoveNewFrameHandleStatsThread()
 	Stats.NewFrameDelegate.Remove( NewFrameDelegateHandle );
 	StatsMasterEnableSubtract();
 	MetadataSize = 0;
-#endif // STATS
+#endif //STATS
 }
 
+
 void FProfilerServiceManager::SetPreviewState( const FMessageAddress& ClientAddress, const bool bRequestedPreviewState )
 {
 #if STATS
@@ -211,7 +191,7 @@ void FProfilerServiceManager::SetPreviewState( const FMessageAddress& ClientAddr
 
 		UE_LOG( LogProfilerService, Verbose, TEXT( "SetPreviewState: %i, InstanceId: %s, ClientAddress: %s" ), (int32)bRequestedPreviewState, *InstanceId.ToString(), *ClientAddress.ToString() );
 	}
-#endif
+#endif //STATS
 }
 
 
@@ -248,12 +228,13 @@ bool FProfilerServiceManager::HandlePing( float DeltaTime )
 		MessageEndpoint->Send(new FProfilerServicePing(), Clients);
 	}
 	return (ClientData.Num() > 0);
-#endif
+#endif //STATS
+
 	return false;
 }
 
 
-void FProfilerServiceManager::HandleServiceCaptureMessage( const FProfilerServiceCapture& Message, const IMessageContextRef& Context )
+void FProfilerServiceManager::HandleServiceCaptureMessage( const FProfilerServiceCapture& Message, const TSharedRef<IMessageContext, ESPMode::ThreadSafe>& Context )
 {
 #if STATS
 	const bool bRequestedCaptureState = Message.bRequestedCaptureState;
@@ -272,11 +253,11 @@ void FProfilerServiceManager::HandleServiceCaptureMessage( const FProfilerServic
 			StopCapture();
 		}
 	}
-#endif
+#endif //STATS
 }
 
 
-void FProfilerServiceManager::HandleServicePongMessage( const FProfilerServicePong& Message, const IMessageContextRef& Context )
+void FProfilerServiceManager::HandleServicePongMessage( const FProfilerServicePong& Message, const TSharedRef<IMessageContext, ESPMode::ThreadSafe>& Context )
 {
 #if STATS
 	FClientData* Data = ClientData.Find(Context->GetSender());
@@ -286,17 +267,17 @@ void FProfilerServiceManager::HandleServicePongMessage( const FProfilerServicePo
 		Data->Active = true;
 		UE_LOG( LogProfilerService, Verbose, TEXT( "Pong InstanceId: %s, GetSender: %s" ), *InstanceId.ToString(), *Context->GetSender().ToString() );
 	}
-#endif
+#endif //STATS
 }
 
 
-void FProfilerServiceManager::HandleServicePreviewMessage( const FProfilerServicePreview& Message, const IMessageContextRef& Context )
+void FProfilerServiceManager::HandleServicePreviewMessage( const FProfilerServicePreview& Message, const TSharedRef<IMessageContext, ESPMode::ThreadSafe>& Context )
 {
 	SetPreviewState( Context->GetSender(), Message.bRequestedPreviewState );
 }
 
 
-void FProfilerServiceManager::HandleServiceRequestMessage( const FProfilerServiceRequest& Message, const IMessageContextRef& Context )
+void FProfilerServiceManager::HandleServiceRequestMessage( const FProfilerServiceRequest& Message, const TSharedRef<IMessageContext, ESPMode::ThreadSafe>& Context )
 {
 	if( Message.Request == EProfilerRequestType::PRT_SendLastCapturedFile )
 	{
@@ -309,8 +290,27 @@ void FProfilerServiceManager::HandleServiceRequestMessage( const FProfilerServic
 }
 
 
+void FProfilerServiceManager::HandleServiceFileChunkMessage(const FProfilerServiceFileChunk& Message, const TSharedRef<IMessageContext, ESPMode::ThreadSafe>& Context)
+{
+	FMemoryReader Reader(Message.Header);
+	FProfilerFileChunkHeader Header;
+	Reader << Header;
+	Header.Validate();
 
-void FProfilerServiceManager::HandleServiceSubscribeMessage( const FProfilerServiceSubscribe& Message, const IMessageContextRef& Context )
+	if (Header.ChunkType == EProfilerFileChunkType::SendChunk)
+	{
+		// Send this file chunk again.
+		FileTransferRunnable->EnqueueFileChunkToSend(new FProfilerServiceFileChunk(Message, FProfilerServiceFileChunk::FNullTag()), true);
+	}
+	else if (Header.ChunkType == EProfilerFileChunkType::FinalizeFile)
+	{
+		// Finalize file.
+		FileTransferRunnable->FinalizeFileSending(Message.Filename);
+	}
+}
+
+
+void FProfilerServiceManager::HandleServiceSubscribeMessage( const FProfilerServiceSubscribe& Message, const TSharedRef<IMessageContext, ESPMode::ThreadSafe>& Context )
 {
 #if STATS
 	const FMessageAddress& SenderAddress = Context->GetSender();
@@ -335,11 +335,11 @@ void FProfilerServiceManager::HandleServiceSubscribeMessage( const FProfilerServ
 			PingDelegateHandle = FTicker::GetCoreTicker().AddTicker(PingDelegate, 5.0f);
 		}
 	}
-#endif
+#endif //STATS
 }
 
 
-void FProfilerServiceManager::HandleServiceUnsubscribeMessage( const FProfilerServiceUnsubscribe& Message, const IMessageContextRef& Context )
+void FProfilerServiceManager::HandleServiceUnsubscribeMessage( const FProfilerServiceUnsubscribe& Message, const TSharedRef<IMessageContext, ESPMode::ThreadSafe>& Context )
 {
 #if	STATS
 	const FMessageAddress SenderAddress = Context->GetSender();
@@ -363,9 +363,10 @@ void FProfilerServiceManager::HandleServiceUnsubscribeMessage( const FProfilerSe
 			FTicker::GetCoreTicker().RemoveTicker(PingDelegateHandle);
 		}
 	}
-#endif // STATS
+#endif //STATS
 }
 
+
 void FProfilerServiceManager::HandleNewFrame(int64 Frame)
 {
 	// Called from the stats thread.
@@ -396,9 +397,10 @@ void FProfilerServiceManager::HandleNewFrame(int64 Frame)
 		FSimpleDelegateGraphTask::FDelegate::CreateRaw( this, &FProfilerServiceManager::CompressDataAndSendToGame, DataToTask, Frame ),
 		TStatId()
 	);
-#endif
+#endif //STATS
 }
 
+
 #if STATS
 
 void FProfilerServiceManager::CompressDataAndSendToGame( TArray<uint8>* DataToTask, int64 Frame )
@@ -435,6 +437,7 @@ void FProfilerServiceManager::CompressDataAndSendToGame( TArray<uint8>* DataToTa
 	delete DataToTask;
 }
 
+
 void FProfilerServiceManager::HandleNewFrameGT( FProfilerServiceData2* ToGameThread )
 {
 	if (MessageEndpoint.IsValid())
@@ -443,4 +446,5 @@ void FProfilerServiceManager::HandleNewFrameGT( FProfilerServiceData2* ToGameThr
 		MessageEndpoint->Send( ToGameThread, PreviewClients );
 	}
 }
-#endif
+
+#endif //STATS
