Fix from xiangguan.nan for avoiding invalid paths.

https://udn.unrealengine.com/questions/496880

diff --git a/Engine/Source/Developer/ProfilerClient/Private/ProfilerClientManager.cpp b/Engine/Source/Developer/ProfilerClient/Private/ProfilerClientManager.cpp
--- a/Engine/Source/Developer/ProfilerClient/Private/ProfilerClientManager.cpp
+++ b/Engine/Source/Developer/ProfilerClient/Private/ProfilerClientManager.cpp
@@ -532,7 +532,7 @@ void FProfilerClientManager::HandleServiceFileChunk(const FProfilerServiceFileCh
 		if(!ReceivedFileInfo)
 		{
 			const FString PathName = FPaths::ProfilingDir() + TEXT("UnrealStats/Received/");
-			const FString StatFilepath = PathName + FileChunk.Filename + StrTmp;
+			const FString StatFilepath = PathName + FPaths::GetCleanFilename((const FString&)FileChunk.Filename) + StrTmp;
 
 			UE_LOG(LogProfilerClient, Log, TEXT("Opening stats file for service-client sending: %s"), *StatFilepath);
 
