Fixing lots of places which created USTRUCTs with operator new and passed them to FMessageEndpoint::Publish or Send, wich crashes when ASAN is used.

The messaging system destroys these objects with FMemory::Free, which has different alignment logic than operator new when ASAN is used, resulting in a crash.

#rnx

diff --git a/Engine/Plugins/Animation/LiveLink/Source/LiveLink/Private/LiveLinkHeartbeatEmitter.cpp b/Engine/Plugins/Animation/LiveLink/Source/LiveLink/Private/LiveLinkHeartbeatEmitter.cpp
--- a/Engine/Plugins/Animation/LiveLink/Source/LiveLink/Private/LiveLinkHeartbeatEmitter.cpp
+++ b/Engine/Plugins/Animation/LiveLink/Source/LiveLink/Private/LiveLinkHeartbeatEmitter.cpp
@@ -58,7 +58,7 @@ uint32 FLiveLinkHeartbeatEmitter::Run()
 			{
 				if (TSharedPtr<FMessageEndpoint, ESPMode::ThreadSafe> MessageEndpoint = Recipient.MessageEndpoint.Pin())
 				{
-					MessageEndpoint->Send(new FLiveLinkHeartbeatMessage(), Recipient.ConnectionAddress);
+					MessageEndpoint->Send(FMessageEndpoint::MakeMessage<FLiveLinkHeartbeatMessage>(), Recipient.ConnectionAddress);
 				}
 			}
 		}
diff --git a/Engine/Plugins/Animation/LiveLink/Source/LiveLink/Private/LiveLinkMessageBusDiscoveryManager.cpp b/Engine/Plugins/Animation/LiveLink/Source/LiveLink/Private/LiveLinkMessageBusDiscoveryManager.cpp
--- a/Engine/Plugins/Animation/LiveLink/Source/LiveLink/Private/LiveLinkMessageBusDiscoveryManager.cpp
+++ b/Engine/Plugins/Animation/LiveLink/Source/LiveLink/Private/LiveLinkMessageBusDiscoveryManager.cpp
@@ -63,7 +63,8 @@ uint32 FLiveLinkMessageBusDiscoveryManager::Run()
 			{
 				LastProviderPoolResults.Reset();
 				LastPingRequest = FGuid::NewGuid();
-				MessageEndpoint->Publish(new FLiveLinkPingMessage(LastPingRequest, ILiveLinkClient::LIVELINK_VERSION));
+				const int32 Version = ILiveLinkClient::LIVELINK_VERSION;
+				MessageEndpoint->Publish(FMessageEndpoint::MakeMessage<FLiveLinkPingMessage>(LastPingRequest, Version));
 			}
 		}
 
diff --git a/Engine/Plugins/Animation/LiveLink/Source/LiveLink/Private/LiveLinkMessageBusFinder.cpp b/Engine/Plugins/Animation/LiveLink/Source/LiveLink/Private/LiveLinkMessageBusFinder.cpp
--- a/Engine/Plugins/Animation/LiveLink/Source/LiveLink/Private/LiveLinkMessageBusFinder.cpp
+++ b/Engine/Plugins/Animation/LiveLink/Source/LiveLink/Private/LiveLinkMessageBusFinder.cpp
@@ -70,7 +70,8 @@ void ULiveLinkMessageBusFinder::PollNetwork()
 
 	PollData.Reset();
 	CurrentPollRequest = FGuid::NewGuid();
-	MessageEndpoint->Publish(new FLiveLinkPingMessage(CurrentPollRequest, ILiveLinkClient::LIVELINK_VERSION));
+	const int32 Version = ILiveLinkClient::LIVELINK_VERSION;
+	MessageEndpoint->Publish(FMessageEndpoint::MakeMessage<FLiveLinkPingMessage>(CurrentPollRequest, Version));
 };
 
 void ULiveLinkMessageBusFinder::HandlePongMessage(const FLiveLinkPongMessage& Message, const TSharedRef<IMessageContext, ESPMode::ThreadSafe>& Context)
diff --git a/Engine/Plugins/FX/Niagara/Source/Niagara/Private/NiagaraDebuggerClient.cpp b/Engine/Plugins/FX/Niagara/Source/Niagara/Private/NiagaraDebuggerClient.cpp
--- a/Engine/Plugins/FX/Niagara/Source/Niagara/Private/NiagaraDebuggerClient.cpp
+++ b/Engine/Plugins/FX/Niagara/Source/Niagara/Private/NiagaraDebuggerClient.cpp
@@ -62,7 +62,7 @@ void FNiagaraDebuggerClient::UpdateClientInfo()
 {
 	if (MessageEndpoint.IsValid() && Connection.IsValid())
 	{
-		FNiagaraSimpleClientInfo* NewInfo = new FNiagaraSimpleClientInfo();
+		FNiagaraSimpleClientInfo* NewInfo = FMessageEndpoint::MakeMessage<FNiagaraSimpleClientInfo>();
 
 		for (TObjectIterator<UNiagaraSystem> It; It; ++It)
 		{
@@ -119,7 +119,7 @@ void FNiagaraDebuggerClient::HandleConnectionRequestMessage(const FNiagaraDebugg
 		
 		//Accept the connection and inform the debugger we have done so with an accepted message.
 		Connection = Context->GetSender();
-		MessageEndpoint->Send(new FNiagaraDebuggerAcceptConnection(SessionId, InstanceId), EMessageFlags::Reliable, nullptr, TArrayBuilder<FMessageAddress>().Add(Connection), FTimespan::Zero(), FDateTime::MaxValue());
+		MessageEndpoint->Send(FMessageEndpoint::MakeMessage<FNiagaraDebuggerAcceptConnection>(SessionId, InstanceId), EMessageFlags::Reliable, nullptr, TArrayBuilder<FMessageAddress>().Add(Connection), FTimespan::Zero(), FDateTime::MaxValue());
 		
 		//Also send an initial update of the client info.
 		UpdateClientInfo();		
@@ -231,7 +231,7 @@ void FNiagaraDebuggerClient::CloseConnection()
 {
 	if (MessageEndpoint.IsValid() && Connection.IsValid())
 	{
-		MessageEndpoint->Send(new FNiagaraDebuggerConnectionClosed(SessionId, InstanceId), EMessageFlags::Reliable, nullptr, TArrayBuilder<FMessageAddress>().Add(Connection), FTimespan::Zero(), FDateTime::MaxValue());
+		MessageEndpoint->Send(FMessageEndpoint::MakeMessage<FNiagaraDebuggerConnectionClosed>(SessionId, InstanceId), EMessageFlags::Reliable, nullptr, TArrayBuilder<FMessageAddress>().Add(Connection), FTimespan::Zero(), FDateTime::MaxValue());
 	}
 
 	OnConnectionClosed();
diff --git a/Engine/Plugins/FX/Niagara/Source/NiagaraEditor/Private/NiagaraDebugger.cpp b/Engine/Plugins/FX/Niagara/Source/NiagaraEditor/Private/NiagaraDebugger.cpp
--- a/Engine/Plugins/FX/Niagara/Source/NiagaraEditor/Private/NiagaraDebugger.cpp
+++ b/Engine/Plugins/FX/Niagara/Source/NiagaraEditor/Private/NiagaraDebugger.cpp
@@ -61,7 +61,7 @@ void FNiagaraDebugger::ExecConsoleCommand(const TCHAR* Cmd, bool bRequiresWorld)
 	auto SendExecCommand = [&](FNiagaraDebugger::FClientInfo& Client)
 	{
 		UE_LOG(LogNiagaraDebugger, Log, TEXT("Sending console command %s. | Session: %s | Instance: %s |"), Cmd, *Client.SessionId.ToString(), *Client.InstanceId.ToString());
-		MessageEndpoint->Send(new FNiagaraDebuggerExecuteConsoleCommand(Cmd, bRequiresWorld), EMessageFlags::Reliable, nullptr, TArrayBuilder<FMessageAddress>().Add(Client.Address), FTimespan::Zero(), FDateTime::MaxValue());
+		MessageEndpoint->Send(FMessageEndpoint::MakeMessage<FNiagaraDebuggerExecuteConsoleCommand>(Cmd, bRequiresWorld), EMessageFlags::Reliable, nullptr, TArrayBuilder<FMessageAddress>().Add(Client.Address), FTimespan::Zero(), FDateTime::MaxValue());
 	};
 
 	ForAllConnectedClients(SendExecCommand);
@@ -75,7 +75,7 @@ void FNiagaraDebugger::UpdateDebugHUDSettings()
 		auto SendSettingsUpdate = [&](FNiagaraDebugger::FClientInfo& Client)
 		{
 			//Create the message and copy the current state of the settings into it.
-			FNiagaraDebugHUDSettingsData* Message = new FNiagaraDebugHUDSettingsData();
+			FNiagaraDebugHUDSettingsData* Message = FMessageEndpoint::MakeMessage<FNiagaraDebugHUDSettingsData>();
 			FNiagaraDebugHUDSettingsData::StaticStruct()->CopyScriptStruct(Message, &Settings->Data);
 
 			UE_LOG(LogNiagaraDebugger, Log, TEXT("Sending updated debug HUD settings. | Session: %s | Instance: %s |"), *Client.SessionId.ToString(), *Client.InstanceId.ToString());
@@ -90,7 +90,7 @@ void FNiagaraDebugger::RequestUpdatedClientInfo()
 {
 	auto RequestUpdate = [&](FNiagaraDebugger::FClientInfo& Client)
 	{
-		FNiagaraRequestSimpleClientInfoMessage* Message = new FNiagaraRequestSimpleClientInfoMessage();
+		FNiagaraRequestSimpleClientInfoMessage* Message = FMessageEndpoint::MakeMessage<FNiagaraRequestSimpleClientInfoMessage>();
 		UE_LOG(LogNiagaraDebugger, Log, TEXT("Requesting updated simple client info. | Session: %s | Instance: %s |"), *Client.SessionId.ToString(), *Client.InstanceId.ToString());
 		MessageEndpoint->Send(Message, EMessageFlags::Reliable, nullptr, TArrayBuilder<FMessageAddress>().Add(Client.Address), FTimespan::Zero(), FDateTime::MaxValue());
 	};
@@ -108,7 +108,7 @@ void FNiagaraDebugger::TriggerOutlinerCapture()
 			auto SendSettingsUpdate = [&](FNiagaraDebugger::FClientInfo& Client)
 			{
 				//Create the message and copy the current state of the settings into it.
-				FNiagaraOutlinerCaptureSettings* Message = new FNiagaraOutlinerCaptureSettings();
+				FNiagaraOutlinerCaptureSettings* Message = FMessageEndpoint::MakeMessage<FNiagaraOutlinerCaptureSettings>();
 				FNiagaraOutlinerCaptureSettings::StaticStruct()->CopyScriptStruct(Message, &Outliner->CaptureSettings);
 
 				UE_LOG(LogNiagaraDebugger, Log, TEXT("Sending updated outliner settings. | Session: %s | Instance: %s |"), *Client.SessionId.ToString(), *Client.InstanceId.ToString());
@@ -196,7 +196,7 @@ void FNiagaraDebugger::SessionManager_OnInstanceSelectionChanged(const TSharedPt
 				return;
 			}
 
-			FNiagaraDebuggerRequestConnection* ConnectionRequestMessage = new FNiagaraDebuggerRequestConnection(SessionId, InstanceId);
+			FNiagaraDebuggerRequestConnection* ConnectionRequestMessage = FMessageEndpoint::MakeMessage<FNiagaraDebuggerRequestConnection>(SessionId, InstanceId);
 			UE_LOG(LogNiagaraDebugger, Log, TEXT("Establishing connection. | Session: %s | Instance: %s (%s)."), *SessionId.ToString(), *InstanceId.ToString(), *Instance->GetInstanceName());
 			MessageEndpoint->Publish(ConnectionRequestMessage);
 
@@ -234,7 +234,7 @@ void FNiagaraDebugger::CloseConnection(FGuid SessionId, FGuid InstanceId)
 	{
 		FClientInfo Pending = PendingClients[FoundPending];
 		PendingClients.RemoveAtSwap(FoundPending);
-		MessageEndpoint->Publish(new FNiagaraDebuggerConnectionClosed(Pending.SessionId, Pending.InstanceId));
+		MessageEndpoint->Publish(FMessageEndpoint::MakeMessage<FNiagaraDebuggerConnectionClosed>(Pending.SessionId, Pending.InstanceId));
 		UE_LOG(LogNiagaraDebugger, Log, TEXT("Closing pending connection. | Session: %s | Instance: %s |"), *SessionId.ToString(), *InstanceId.ToString());
 
 		OnConnectionClosedDelegate.Broadcast(Pending);
@@ -243,7 +243,7 @@ void FNiagaraDebugger::CloseConnection(FGuid SessionId, FGuid InstanceId)
 	{
 		FClientInfo Active = ConnectedClients[FoundActive];
 		ConnectedClients.RemoveAtSwap(FoundActive);
-		MessageEndpoint->Send(new FNiagaraDebuggerConnectionClosed(Active.SessionId, Active.InstanceId), EMessageFlags::Reliable, nullptr, TArrayBuilder<FMessageAddress>().Add(Active.Address), FTimespan::Zero(), FDateTime::MaxValue());
+		MessageEndpoint->Send(FMessageEndpoint::MakeMessage<FNiagaraDebuggerConnectionClosed>(Active.SessionId, Active.InstanceId), EMessageFlags::Reliable, nullptr, TArrayBuilder<FMessageAddress>().Add(Active.Address), FTimespan::Zero(), FDateTime::MaxValue());
 		UE_LOG(LogNiagaraDebugger, Log, TEXT("Closing active connection. | Session: %s | Instance: %s |"), *SessionId.ToString(), *InstanceId.ToString());
 
 		OnConnectionClosedDelegate.Broadcast(Active);
diff --git a/Engine/Source/Developer/AutomationController/Private/AutomationControllerManger.cpp b/Engine/Source/Developer/AutomationController/Private/AutomationControllerManger.cpp
--- a/Engine/Source/Developer/AutomationController/Private/AutomationControllerManger.cpp
+++ b/Engine/Source/Developer/AutomationController/Private/AutomationControllerManger.cpp
@@ -116,7 +116,7 @@ void FAutomationControllerManager::RequestAvailableWorkers(const FGuid& SessionI
 	int32 ChangelistNumber = 10000;
 	FString ProcessName = TEXT("instance_name");
 
-	MessageEndpoint->Publish(new FAutomationWorkerFindWorkers(ChangelistNumber, FApp::GetProjectName(), ProcessName, SessionId), EMessageScope::Network);
+	MessageEndpoint->Publish(FMessageEndpoint::MakeMessage<FAutomationWorkerFindWorkers>(ChangelistNumber, FApp::GetProjectName(), ProcessName, SessionId), EMessageScope::Network);
 
 	// Reset the check test timers
 	LastTimeUpdateTicked = FPlatformTime::Seconds();
@@ -145,7 +145,7 @@ void FAutomationControllerManager::RequestTests()
 			UE_LOG(LogAutomationController, Log, TEXT("Requesting test list from %s"), *MessageAddress.ToString());
 
 			//issue tests on appropriate platforms
-			MessageEndpoint->Send(new FAutomationWorkerRequestTests(bDeveloperDirectoryIncluded, RequestedTestFlags), MessageAddress);
+			MessageEndpoint->Send(FMessageEndpoint::MakeMessage<FAutomationWorkerRequestTests>(bDeveloperDirectoryIncluded, RequestedTestFlags), MessageAddress);
 		}
 	}
 }
@@ -203,7 +203,7 @@ void FAutomationControllerManager::RunTests(const bool bInIsLocalSession)
 			FMessageAddress MessageAddress = DeviceClusterManager.GetDeviceMessageAddress(ClusterIndex, DeviceIndex);
 			UE_LOG(LogAutomationController, Log, TEXT("Sending Reset Tests to %s"), *MessageAddress.ToString());
 
-			MessageEndpoint->Send(new FAutomationWorkerResetTests(), MessageAddress);
+			MessageEndpoint->Send(FMessageEndpoint::MakeMessage<FAutomationWorkerResetTests>(), MessageAddress);
 		}
 	}
 	
@@ -242,7 +242,7 @@ void FAutomationControllerManager::StopTests()
 				FMessageAddress MessageAddress = DeviceClusterManager.GetDeviceMessageAddress(ClusterIndex, DeviceIndex);
 
 				UE_LOG(LogAutomationController, Log, TEXT("Sending StopTests to %s"), *MessageAddress.ToString());
-				MessageEndpoint->Send(new FAutomationWorkerStopTests(), MessageAddress);
+				MessageEndpoint->Send(FMessageEndpoint::MakeMessage<FAutomationWorkerStopTests>(), MessageAddress);
 			}
 		}
 
@@ -268,7 +268,7 @@ void FAutomationControllerManager::Init()
 
 void FAutomationControllerManager::RequestLoadAsset(const FString& InAssetName)
 {
-	MessageEndpoint->Publish(new FAssetEditorRequestOpenAsset(InAssetName), EMessageScope::Process);
+	MessageEndpoint->Publish(FMessageEndpoint::MakeMessage<FAssetEditorRequestOpenAsset>(InAssetName), EMessageScope::Process);
 }
 
 void FAutomationControllerManager::Tick()
@@ -601,7 +601,7 @@ void FAutomationControllerManager::ExecuteNextTask( int32 ClusterIndex, OUT bool
 							// Send the test to the device for execution!
 							UE_LOG(LogAutomationController, Log, TEXT("Sending RunTest %s to %s"), *NextTest->GetDisplayName(), *DeviceAddress.ToString());
 
-							MessageEndpoint->Send(new FAutomationWorkerRunTests(ExecutionCount, AddressIndex, NextTest->GetCommand(), NextTest->GetDisplayName(), bSendAnalytics), DeviceAddress);
+							MessageEndpoint->Send(FMessageEndpoint::MakeMessage<FAutomationWorkerRunTests>(ExecutionCount, AddressIndex, NextTest->GetCommand(), NextTest->GetDisplayName(), bSendAnalytics), DeviceAddress);
 
 							// Add a test so we can check later if the device is still active
 							TestRunningArray.Add(FTestRunningInfo(DeviceAddress));
@@ -968,7 +968,7 @@ void FAutomationControllerManager::UpdateTests()
 				}
 				else
 				{
-					MessageEndpoint->Send(new FAutomationWorkerPing(), TestRunningArray[Index].OwnerMessageAddress);
+					MessageEndpoint->Send(FMessageEndpoint::MakeMessage<FAutomationWorkerPing>(), TestRunningArray[Index].OwnerMessageAddress);
 				}
 			}
 			CheckTestTimer = 0.f;
@@ -1180,7 +1180,7 @@ void FAutomationControllerManager::HandleRequestNextNetworkCommandMessage(const
 			for ( int32 AddressIndex = 0; AddressIndex < DeviceAddresses.Num(); ++AddressIndex )
 			{
 				//send "next command message" to worker
-				MessageEndpoint->Send(new FAutomationWorkerNextNetworkCommandReply(), DeviceAddresses[AddressIndex]);
+				MessageEndpoint->Send(FMessageEndpoint::MakeMessage<FAutomationWorkerNextNetworkCommandReply>(), DeviceAddresses[AddressIndex]);
 			}
 		}
 	}
diff --git a/Engine/Source/Developer/ProfilerClient/Private/ProfilerClientManager.cpp b/Engine/Source/Developer/ProfilerClient/Private/ProfilerClientManager.cpp
--- a/Engine/Source/Developer/ProfilerClient/Private/ProfilerClientManager.cpp
+++ b/Engine/Source/Developer/ProfilerClient/Private/ProfilerClientManager.cpp
@@ -87,7 +87,7 @@ void FProfilerClientManager::Subscribe(const FGuid& Session)
 			Connections.GenerateKeyArray(Instances);
 			for (int32 i = 0; i < Instances.Num(); ++i)
 			{
-				MessageEndpoint->Publish(new FProfilerServiceUnsubscribe(OldSessionId, Instances[i]), EMessageScope::Network);
+				MessageEndpoint->Publish(FMessageEndpoint::MakeMessage<FProfilerServiceUnsubscribe>(OldSessionId, Instances[i]), EMessageScope::Network);
 
 				// fire the disconnection delegate
 				ProfilerClientDisconnectedDelegate.Broadcast(ActiveSessionId, Instances[i]);
@@ -112,7 +112,7 @@ void FProfilerClientManager::Track(const FGuid& Instance)
 	{
 		PendingInstances.Add(Instance);
 
-		MessageEndpoint->Publish(new FProfilerServiceSubscribe(ActiveSessionId, Instance), EMessageScope::Network);
+		MessageEndpoint->Publish(FMessageEndpoint::MakeMessage<FProfilerServiceSubscribe>(ActiveSessionId, Instance), EMessageScope::Network);
 
 		RetryTime = 5.f;
 		TickDelegateHandle = FTicker::GetCoreTicker().AddTicker(TickDelegate, RetryTime);
@@ -128,7 +128,7 @@ void FProfilerClientManager::Untrack(const FGuid& Instance)
 #if STATS
 	if (MessageEndpoint.IsValid() && ActiveSessionId.IsValid())
 	{
-		MessageEndpoint->Publish(new FProfilerServiceUnsubscribe(ActiveSessionId, Instance), EMessageScope::Network);
+		MessageEndpoint->Publish(FMessageEndpoint::MakeMessage<FProfilerServiceUnsubscribe>(ActiveSessionId, Instance), EMessageScope::Network);
 		Connections.Remove(Instance);
 
 		// fire the disconnection delegate
@@ -161,13 +161,13 @@ void FProfilerClientManager::SetCaptureState(const bool bRequestedCaptureState,
 			{
 				Instances.Add(It.Value().ProfilerServiceAddress);
 			}
-			MessageEndpoint->Send(new FProfilerServiceCapture(bRequestedCaptureState), Instances);
+			MessageEndpoint->Send(FMessageEndpoint::MakeMessage<FProfilerServiceCapture>(bRequestedCaptureState), Instances);
 			UE_LOG(LogProfilerClient, Verbose, TEXT("SetCaptureState Session: %s, Instance: %s, State: %i"), *ActiveSessionId.ToString(), *InstanceId.ToString(), (int32)bRequestedCaptureState);
 		}
 		else
 		{
 			const FMessageAddress& MessageAddress = Connections.Find(InstanceId)->ProfilerServiceAddress;
-			MessageEndpoint->Send(new FProfilerServiceCapture(bRequestedCaptureState), MessageAddress);
+			MessageEndpoint->Send(FMessageEndpoint::MakeMessage<FProfilerServiceCapture>(bRequestedCaptureState), MessageAddress);
 
 			UE_LOG(LogProfilerClient, Verbose, TEXT("SetCaptureState Session: %s, Instance: %s, State: %i"), *ActiveSessionId.ToString(), *InstanceId.ToString(), (int32)bRequestedCaptureState);
 		}
@@ -188,13 +188,13 @@ void FProfilerClientManager::SetPreviewState(const bool bRequestedPreviewState,
 			{
 				Instances.Add(It.Value().ProfilerServiceAddress);
 			}
-			MessageEndpoint->Send(new FProfilerServicePreview(bRequestedPreviewState), Instances);
+			MessageEndpoint->Send(FMessageEndpoint::MakeMessage<FProfilerServicePreview>(bRequestedPreviewState), Instances);
 			UE_LOG(LogProfilerClient, Verbose, TEXT("SetPreviewState Session: %s, Instance: %s, State: %i"), *ActiveSessionId.ToString(), *InstanceId.ToString(), (int32)bRequestedPreviewState);
 		}
 		else
 		{
 			const FMessageAddress& MessageAddress = Connections.Find(InstanceId)->ProfilerServiceAddress;
-			MessageEndpoint->Send(new FProfilerServicePreview(bRequestedPreviewState), MessageAddress);
+			MessageEndpoint->Send(FMessageEndpoint::MakeMessage<FProfilerServicePreview>(bRequestedPreviewState), MessageAddress);
 			UE_LOG(LogProfilerClient, Verbose, TEXT("SetPreviewState Session: %s, Instance: %s, State: %i"), *ActiveSessionId.ToString(), *InstanceId.ToString(), (int32)bRequestedPreviewState);
 		}
 	}
@@ -332,12 +332,12 @@ void FProfilerClientManager::RequestLastCapturedFile(const FGuid& InstanceId /*=
 			{
 				Instances.Add(It.Value().ProfilerServiceAddress);
 			}
-			MessageEndpoint->Send(new FProfilerServiceRequest(EProfilerRequestType::PRT_SendLastCapturedFile), Instances);
+			MessageEndpoint->Send(FMessageEndpoint::MakeMessage<FProfilerServiceRequest>(EProfilerRequestType::PRT_SendLastCapturedFile), Instances);
 		}
 		else
 		{
 			const FMessageAddress& MessageAddress = Connections.Find(InstanceId)->ProfilerServiceAddress;
-			MessageEndpoint->Send(new FProfilerServiceRequest(EProfilerRequestType::PRT_SendLastCapturedFile), MessageAddress);
+			MessageEndpoint->Send(FMessageEndpoint::MakeMessage<FProfilerServiceRequest>(EProfilerRequestType::PRT_SendLastCapturedFile), MessageAddress);
 		}
 	}
 #endif
@@ -563,7 +563,7 @@ void FProfilerClientManager::HandleServiceFileChunk(const FProfilerServiceFileCh
 				// File has been successfully sent, so send this information to the profiler service.
 				if(MessageEndpoint.IsValid())
 				{
-					MessageEndpoint->Send(new FProfilerServiceFileChunk(FGuid(),FileChunk.Filename,FProfilerFileChunkHeader(0,0,0,EProfilerFileChunkType::FinalizeFile).AsArray()), Context->GetSender());
+					MessageEndpoint->Send(FMessageEndpoint::MakeMessage<FProfilerServiceFileChunk>(FGuid(),FileChunk.Filename,FProfilerFileChunkHeader(0,0,0,EProfilerFileChunkType::FinalizeFile).AsArray()), Context->GetSender());
 					ProfilerFileTransferDelegate.Broadcast(FileChunk.Filename, ReceivedFileInfo->Progress, FileChunkHeader.FileSize);
 				}
 				
@@ -588,7 +588,7 @@ void FProfilerClientManager::HandleServiceFileChunk(const FProfilerServiceFileCh
 			// This chunk is a bad chunk, so ask for resending it.
 			if(MessageEndpoint.IsValid())
 			{
-				MessageEndpoint->Send(new FProfilerServiceFileChunk(FileChunk,FProfilerServiceFileChunk::FNullTag()), Context->GetSender());
+				MessageEndpoint->Send(FMessageEndpoint::MakeMessage<FProfilerServiceFileChunk>(FileChunk,FProfilerServiceFileChunk::FNullTag()), Context->GetSender());
 				UE_LOG(LogProfilerClient, Log, TEXT("Received a bad chunk of file, resending: %5i, %6u, %10u, %s"), FileChunk.HexData.Len(), ReceivedFileInfo->Progress, FileChunkHeader.FileSize, *FileChunk.Filename);
 			}
 		}
@@ -607,7 +607,7 @@ void FProfilerClientManager::HandleServicePingMessage(const FProfilerServicePing
 		{
 			Instances.Add(It.Value().ProfilerServiceAddress);
 		}
-		MessageEndpoint->Send(new FProfilerServicePong(), Instances);
+		MessageEndpoint->Send(FMessageEndpoint::MakeMessage<FProfilerServicePong>(), Instances);
 
 		UE_LOG(LogProfilerClient, Verbose, TEXT("Ping GetSender: %s"), *Context->GetSender().ToString());
 	}
diff --git a/Engine/Source/Developer/ProfilerService/Private/ProfilerServiceManager.cpp b/Engine/Source/Developer/ProfilerService/Private/ProfilerServiceManager.cpp
--- a/Engine/Source/Developer/ProfilerService/Private/ProfilerServiceManager.cpp
+++ b/Engine/Source/Developer/ProfilerService/Private/ProfilerServiceManager.cpp
@@ -182,7 +182,7 @@ void FProfilerServiceManager::SetPreviewState( const FMessageAddress& ClientAddr
 				PreviewClients.Add(ClientAddress);
 				Client->Preview = true;
 
-				MessageEndpoint->Send( new FProfilerServicePreviewAck( InstanceId ), ClientAddress );
+				MessageEndpoint->Send(FMessageEndpoint::MakeMessage<FProfilerServicePreviewAck>(InstanceId), ClientAddress);
 			}
 			else
 			{
@@ -248,7 +248,7 @@ bool FProfilerServiceManager::HandlePing( float DeltaTime )
 	// send the ping message
 	if (MessageEndpoint.IsValid() && Clients.Num() > 0)
 	{
-		MessageEndpoint->Send(new FProfilerServicePing(), Clients);
+		MessageEndpoint->Send(FMessageEndpoint::MakeMessage<FProfilerServicePing>(), Clients);
 	}
 	return (ClientData.Num() > 0);
 #endif //STATS
@@ -349,7 +349,7 @@ void FProfilerServiceManager::HandleServiceSubscribeMessage( const FProfilerServ
 		ClientData.Add( SenderAddress, Data );
 
 		// Send authorize.
-		MessageEndpoint->Send( new FProfilerServiceAuthorize( SessionId, InstanceId ), SenderAddress );
+		MessageEndpoint->Send(FMessageEndpoint::MakeMessage<FProfilerServiceAuthorize>(SessionId, InstanceId), SenderAddress);
 		// Eventually send the metadata if needed.
 
 		// Initiate the ping callback
diff --git a/Engine/Source/Developer/SlateReflector/Private/WidgetSnapshotService.cpp b/Engine/Source/Developer/SlateReflector/Private/WidgetSnapshotService.cpp
--- a/Engine/Source/Developer/SlateReflector/Private/WidgetSnapshotService.cpp
+++ b/Engine/Source/Developer/SlateReflector/Private/WidgetSnapshotService.cpp
@@ -29,7 +29,7 @@ FGuid FWidgetSnapshotService::RequestSnapshot(const FGuid& InRemoteInstanceId, c
 {
 	if (MessageEndpoint.IsValid())
 	{
-		FWidgetSnapshotRequest* WidgetSnapshotRequest = new FWidgetSnapshotRequest;
+		FWidgetSnapshotRequest* WidgetSnapshotRequest = FMessageEndpoint::MakeMessage<FWidgetSnapshotRequest>();
 		WidgetSnapshotRequest->TargetInstanceId = InRemoteInstanceId;
 		WidgetSnapshotRequest->SnapshotRequestId = FGuid::NewGuid();
 
diff --git a/Engine/Source/Developer/TargetDeviceServices/Private/Proxies/TargetDeviceProxy.cpp b/Engine/Source/Developer/TargetDeviceServices/Private/Proxies/TargetDeviceProxy.cpp
--- a/Engine/Source/Developer/TargetDeviceServices/Private/Proxies/TargetDeviceProxy.cpp
+++ b/Engine/Source/Developer/TargetDeviceServices/Private/Proxies/TargetDeviceProxy.cpp
@@ -261,7 +261,7 @@ bool FTargetDeviceProxy::HasVariant(FName InVariant) const
 
 bool FTargetDeviceProxy::TerminateLaunchedProcess(FName InVariant, const FString& ProcessIdentifier)
 {
-	MessageEndpoint->Send(new FTargetDeviceServiceTerminateLaunchedProcess(InVariant, ProcessIdentifier), MessageAddress);
+	MessageEndpoint->Send(FMessageEndpoint::MakeMessage<FTargetDeviceServiceTerminateLaunchedProcess>(InVariant, ProcessIdentifier), MessageAddress);
 
 	return true;
 }
@@ -269,19 +269,19 @@ bool FTargetDeviceProxy::TerminateLaunchedProcess(FName InVariant, const FString
 
 void FTargetDeviceProxy::PowerOff(bool Force)
 {
-	MessageEndpoint->Send(new FTargetDeviceServicePowerOff(FPlatformProcess::UserName(false), Force), MessageAddress);
+	MessageEndpoint->Send(FMessageEndpoint::MakeMessage<FTargetDeviceServicePowerOff>(FPlatformProcess::UserName(false), Force), MessageAddress);
 }
 
 
 void FTargetDeviceProxy::PowerOn()
 {
-	MessageEndpoint->Send(new FTargetDeviceServicePowerOn(FPlatformProcess::UserName(false)), MessageAddress);
+	MessageEndpoint->Send(FMessageEndpoint::MakeMessage<FTargetDeviceServicePowerOn>(FPlatformProcess::UserName(false)), MessageAddress);
 }
 
 
 void FTargetDeviceProxy::Reboot()
 {
-	MessageEndpoint->Send(new FTargetDeviceServiceReboot(FPlatformProcess::UserName(false)), MessageAddress);
+	MessageEndpoint->Send(FMessageEndpoint::MakeMessage<FTargetDeviceServiceReboot>(FPlatformProcess::UserName(false)), MessageAddress);
 }
 
 
diff --git a/Engine/Source/Developer/TargetDeviceServices/Private/Proxies/TargetDeviceProxyManager.cpp b/Engine/Source/Developer/TargetDeviceServices/Private/Proxies/TargetDeviceProxyManager.cpp
--- a/Engine/Source/Developer/TargetDeviceServices/Private/Proxies/TargetDeviceProxyManager.cpp
+++ b/Engine/Source/Developer/TargetDeviceServices/Private/Proxies/TargetDeviceProxyManager.cpp
@@ -148,10 +148,8 @@ void FTargetDeviceProxyManager::SendPing()
 			Proxy.Value->AddTimeout(Timeout);
 		}
 
-		// message is going to be deleted by FMemory::Free() (see FMessageContext destructor), so allocate it with Malloc
-		void* Memory = FMemory::Malloc(sizeof(FTargetDeviceServicePing), alignof(FTargetDeviceServicePing));
 		// send to only this process, as we no longer want to send out pings to the local network and get another machine's list of devices
-		MessageEndpoint->Publish(new (Memory) FTargetDeviceServicePing(FPlatformProcess::UserName(false)), EMessageScope::Process);
+		MessageEndpoint->Publish(FMessageEndpoint::MakeMessage<FTargetDeviceServicePing>(FPlatformProcess::UserName(false)), EMessageScope::Process);
 	}
 }
 
diff --git a/Engine/Source/Developer/TargetDeviceServices/Private/Services/TargetDeviceService.cpp b/Engine/Source/Developer/TargetDeviceServices/Private/Services/TargetDeviceService.cpp
--- a/Engine/Source/Developer/TargetDeviceServices/Private/Services/TargetDeviceService.cpp
+++ b/Engine/Source/Developer/TargetDeviceServices/Private/Services/TargetDeviceService.cpp
@@ -242,9 +242,7 @@ bool FTargetDeviceService::Start()
 		ClaimHost = FPlatformProcess::ComputerName();
 		ClaimUser = FPlatformProcess::UserName(false);
 
-		// message is going to be deleted by FMemory::Free() (see FMessageContext destructor), so allocate it with Malloc
-		void* Memory = FMemory::Malloc(sizeof(FTargetDeviceClaimed), alignof(FTargetDeviceClaimed));
-		MessageEndpoint->Publish(new(Memory) FTargetDeviceClaimed(DeviceName, ClaimHost, ClaimUser));
+		MessageEndpoint->Publish(FMessageEndpoint::MakeMessage<FTargetDeviceClaimed>(DeviceName, ClaimHost, ClaimUser));
 
 		Running = true;
 	}
@@ -258,8 +256,7 @@ void FTargetDeviceService::Stop()
 	if (Running)
 	{
 		// message is going to be deleted by FMemory::Free() (see FMessageContext destructor), so allocate it with Malloc
-		void* Memory = FMemory::Malloc(sizeof(FTargetDeviceUnclaimed), alignof(FTargetDeviceUnclaimed));
-		MessageEndpoint->Publish(new(Memory) FTargetDeviceUnclaimed(DeviceName, FPlatformProcess::ComputerName(), FPlatformProcess::UserName(false)));
+		MessageEndpoint->Publish(FMessageEndpoint::MakeMessage<FTargetDeviceUnclaimed>(DeviceName, FPlatformProcess::ComputerName(), FPlatformProcess::UserName(false)));
 		FPlatformProcess::SleepNoStats(0.01);
 
 		// Only stop the device if we care about device claiming
@@ -356,9 +353,7 @@ void FTargetDeviceService::HandleClaimedMessage(const FTargetDeviceClaimed& Mess
 	{
 		if (Context->GetSender() != MessageEndpoint->GetAddress())
 		{
-			// message is going to be deleted by FMemory::Free() (see FMessageContext destructor), so allocate it with Malloc
-			void* Memory = FMemory::Malloc(sizeof(FTargetDeviceClaimDenied), alignof(FTargetDeviceClaimDenied));
-			MessageEndpoint->Send(new(Memory) FTargetDeviceClaimDenied(DeviceName, FPlatformProcess::ComputerName(), FPlatformProcess::UserName(false)), Context->GetSender());
+			MessageEndpoint->Send(FMessageEndpoint::MakeMessage<FTargetDeviceClaimDenied>(DeviceName, FPlatformProcess::ComputerName(), FPlatformProcess::UserName(false)), Context->GetSender());
 		}
 	}
 	else
diff --git a/Engine/Source/Editor/SwarmInterface/Private/SwarmInterfaceLocal.cpp b/Engine/Source/Editor/SwarmInterface/Private/SwarmInterfaceLocal.cpp
--- a/Engine/Source/Editor/SwarmInterface/Private/SwarmInterfaceLocal.cpp
+++ b/Engine/Source/Editor/SwarmInterface/Private/SwarmInterfaceLocal.cpp
@@ -190,7 +190,7 @@ int32 FSwarmInterfaceLocalImpl::OpenConnection( FConnectionCallback InCallbackFu
 		if (MessageEndpoint.IsValid())
 		{
 			MessageEndpoint->Subscribe<FSwarmPingMessage>();
-			MessageEndpoint->Publish(new FSwarmPingMessage(), EMessageScope::Network);
+			MessageEndpoint->Publish(FMessageEndpoint::MakeMessage<FSwarmPingMessage>(), EMessageScope::Network);
 
 			// UMB does not allow us to identify early its initialization errors - check that manually.
 			bIsConnected = SwarmInterfaceLocalImpl::CanUseUMB();
@@ -234,7 +234,7 @@ int32 FSwarmInterfaceLocalImpl::SendMessage( const FMessage& Message )
 
 	while (bIsConnected && !Recepient.IsValid())
 	{
-		MessageEndpoint->Publish(new FSwarmPingMessage(), EMessageScope::Network);
+		MessageEndpoint->Publish(FMessageEndpoint::MakeMessage<FSwarmPingMessage>(), EMessageScope::Network);
 		
 		FTaskGraphInterface::Get().ProcessThreadUntilIdle(ENamedThreads::GameThread);
 		FPlatformProcess::Sleep(0.5f);
@@ -250,17 +250,17 @@ int32 FSwarmInterfaceLocalImpl::SendMessage( const FMessage& Message )
 	if (Message.Type == MESSAGE_INFO)
 	{
 		FInfoMessage* InfoMessage = (FInfoMessage*)&Message;
-		MessageEndpoint->Send(new FSwarmInfoMessage(InfoMessage->TextMessage), Recepient);
+		MessageEndpoint->Send(FMessageEndpoint::MakeMessage<FSwarmInfoMessage>(InfoMessage->TextMessage), Recepient);
 	}
 	else if (Message.Type == MESSAGE_ALERT)
 	{
 		FAlertMessage* AlertMessage = (FAlertMessage*)&Message;
-		MessageEndpoint->Send(new FSwarmAlertMessage(AlertMessage->JobGuid, AlertMessage->AlertLevel, AlertMessage->ObjectGuid, AlertMessage->TypeId, AlertMessage->TextMessage), Recepient);
+		MessageEndpoint->Send(FMessageEndpoint::MakeMessage<FSwarmAlertMessage>(AlertMessage->JobGuid, AlertMessage->AlertLevel, AlertMessage->ObjectGuid, AlertMessage->TypeId, AlertMessage->TextMessage), Recepient);
 	}
 	else if (Message.Type == MESSAGE_TIMING)
 	{
 		FTimingMessage* TimingMessage = (FTimingMessage*)&Message;
-		MessageEndpoint->Send(new FSwarmTimingMessage(TimingMessage->State, TimingMessage->ThreadNum), Recepient);
+		MessageEndpoint->Send(FMessageEndpoint::MakeMessage<FSwarmTimingMessage>(TimingMessage->State, TimingMessage->ThreadNum), Recepient);
 	}
 	else if (Message.Type == MESSAGE_TASK_REQUEST)
 	{
@@ -275,18 +275,18 @@ int32 FSwarmInterfaceLocalImpl::SendMessage( const FMessage& Message )
 			FTaskSpecification TaskSpec = Tasks.Pop();
 			CallbackFunc( (FMessage*)&TaskSpec, CallbackData );
 		}
-		MessageEndpoint->Send(new FSwarmTaskRequestMessage(), Recepient);
+		MessageEndpoint->Send(FMessageEndpoint::MakeMessage<FSwarmTaskRequestMessage>(), Recepient);
 	}
 	else if (Message.Type == MESSAGE_TASK_REQUEST_RESPONSE)
 	{
 		FTaskRequestResponse* ResponseMessage = (FTaskRequestResponse*)&Message;
 		if (ResponseMessage->ResponseType == RESPONSE_TYPE_RELEASE)
 		{
-			MessageEndpoint->Send(new FSwarmTaskRequestReleaseMessage(), Recepient);
+			MessageEndpoint->Send(FMessageEndpoint::MakeMessage<FSwarmTaskRequestReleaseMessage>(), Recepient);
 		}
 		else if (ResponseMessage->ResponseType == RESPONSE_TYPE_RESERVATION)
 		{
-			MessageEndpoint->Send(new FSwarmTaskRequestReservationMessage(), Recepient);
+			MessageEndpoint->Send(FMessageEndpoint::MakeMessage<FSwarmTaskRequestReservationMessage>(), Recepient);
 		}
 		else if (ResponseMessage->ResponseType == RESPONSE_TYPE_SPECIFICATION)
 		{
@@ -296,22 +296,22 @@ int32 FSwarmInterfaceLocalImpl::SendMessage( const FMessage& Message )
 			{
 				Dependencies.Add(SpecificationMessage->Dependencies[Index]);
 			}
-			MessageEndpoint->Send(new FSwarmTaskRequestSpecificationMessage(SpecificationMessage->TaskGuid, SpecificationMessage->Parameters, SpecificationMessage->Flags, Dependencies), Recepient);
+			MessageEndpoint->Send(FMessageEndpoint::MakeMessage<FSwarmTaskRequestSpecificationMessage>(SpecificationMessage->TaskGuid, SpecificationMessage->Parameters, SpecificationMessage->Flags, Dependencies), Recepient);
 		}
 	}
 	else if (Message.Type == MESSAGE_JOB_STATE)
 	{
 		FJobState* StateMessage = (FJobState*)&Message;
-		MessageEndpoint->Send(new FSwarmJobStateMessage(StateMessage->JobGuid, StateMessage->JobState, StateMessage->JobMessage, StateMessage->JobExitCode, StateMessage->JobRunningTime), Recepient);
+		MessageEndpoint->Send(FMessageEndpoint::MakeMessage<FSwarmJobStateMessage>(StateMessage->JobGuid, StateMessage->JobState, StateMessage->JobMessage, StateMessage->JobExitCode, StateMessage->JobRunningTime), Recepient);
 	}
 	else if (Message.Type == MESSAGE_TASK_STATE)
 	{
 		FTaskState* StateMessage = (FTaskState*)&Message;
-		MessageEndpoint->Send(new FSwarmTaskStateMessage(StateMessage->TaskGuid, StateMessage->TaskState, StateMessage->TaskMessage, StateMessage->TaskExitCode, StateMessage->TaskRunningTime), Recepient);
+		MessageEndpoint->Send(FMessageEndpoint::MakeMessage<FSwarmTaskStateMessage>(StateMessage->TaskGuid, StateMessage->TaskState, StateMessage->TaskMessage, StateMessage->TaskExitCode, StateMessage->TaskRunningTime), Recepient);
 	}
 	else if (Message.Type == MESSAGE_QUIT)
 	{
-		MessageEndpoint->Send(new FSwarmQuitMessage(), Recepient);
+		MessageEndpoint->Send(FMessageEndpoint::MakeMessage<FSwarmQuitMessage>(), Recepient);
 	}
 #endif
 	return 0;
@@ -320,7 +320,7 @@ int32 FSwarmInterfaceLocalImpl::SendMessage( const FMessage& Message )
 #if USE_LOCAL_SWARM_INTERFACE
 void FSwarmInterfaceLocalImpl::HandlePingMessage( const FSwarmPingMessage& Message, const TSharedRef<IMessageContext, ESPMode::ThreadSafe>& Context )
 {
-	MessageEndpoint->Send(new FSwarmPongMessage(bIsEditor, FPlatformProcess::ComputerName()), Context->GetSender());
+	MessageEndpoint->Send(FMessageEndpoint::MakeMessage<FSwarmPongMessage>(bIsEditor, FPlatformProcess::ComputerName()), Context->GetSender());
 }
 
 void FSwarmInterfaceLocalImpl::HandlePongMessage( const FSwarmPongMessage& Message, const TSharedRef<IMessageContext, ESPMode::ThreadSafe>& Context )
diff --git a/Engine/Source/Editor/UnrealEd/Private/Commandlets/FileServerCommandlet.cpp b/Engine/Source/Editor/UnrealEd/Private/Commandlets/FileServerCommandlet.cpp
--- a/Engine/Source/Editor/UnrealEd/Private/Commandlets/FileServerCommandlet.cpp
+++ b/Engine/Source/Editor/UnrealEd/Private/Commandlets/FileServerCommandlet.cpp
@@ -89,7 +89,7 @@ int32 UFileServerCommandlet::Main( const FString& Params )
 
 		if (MessageEndpoint.IsValid())
 		{
-			MessageEndpoint->Publish(new FFileServerReady(AddressStringList, InstanceId), EMessageScope::Network);
+			MessageEndpoint->Publish(FMessageEndpoint::MakeMessage<FFileServerReady>(AddressStringList, InstanceId), EMessageScope::Network);
 		}		
 	}
 
diff --git a/Engine/Source/Editor/UnrealEd/Private/CookOnTheFlyServer.cpp b/Engine/Source/Editor/UnrealEd/Private/CookOnTheFlyServer.cpp
--- a/Engine/Source/Editor/UnrealEd/Private/CookOnTheFlyServer.cpp
+++ b/Engine/Source/Editor/UnrealEd/Private/CookOnTheFlyServer.cpp
@@ -975,7 +975,7 @@ bool UCookOnTheFlyServer::BroadcastFileserverPresence( const FGuid &InstanceId )
 
 	if (MessageEndpoint.IsValid())
 	{
-		MessageEndpoint->Publish(new FFileServerReady(AddressStringList, InstanceId), EMessageScope::Network);
+		MessageEndpoint->Publish(FMessageEndpoint::MakeMessage<FFileServerReady>(AddressStringList, InstanceId), EMessageScope::Network);
 	}		
 	
 	return true;
diff --git a/Engine/Source/Programs/IOS/UnrealLaunchDaemon/Private/LaunchDaemonMessageHandler.cpp b/Engine/Source/Programs/IOS/UnrealLaunchDaemon/Private/LaunchDaemonMessageHandler.cpp
--- a/Engine/Source/Programs/IOS/UnrealLaunchDaemon/Private/LaunchDaemonMessageHandler.cpp
+++ b/Engine/Source/Programs/IOS/UnrealLaunchDaemon/Private/LaunchDaemonMessageHandler.cpp
@@ -39,7 +39,7 @@ void FLaunchDaemonMessageHandler::HandlePingMessage(const FIOSLaunchDaemonPing&
 	{
 		FMessageAddress MessageSender = Context->GetSender();
 
-		MessageEndpoint->Send(new FIOSLaunchDaemonPong(FString(FPlatformProperties::PlatformName()) + (TARGET_IPHONE_SIMULATOR ? FString(TEXT("Simulator:")) : FString(TEXT("@"))) + FString(FPlatformProcess::ComputerName())
+		MessageEndpoint->Send(FMessageEndpoint::MakeMessage<FIOSLaunchDaemonPong>(FString(FPlatformProperties::PlatformName()) + (TARGET_IPHONE_SIMULATOR ? FString(TEXT("Simulator:")) : FString(TEXT("@"))) + FString(FPlatformProcess::ComputerName())
 			, FPlatformProcess::ComputerName()
 			, "ULD_Ready"
 			, UI_USER_INTERFACE_IDIOM() == UIUserInterfaceIdiomPhone? "Phone" : "Tablet"
diff --git a/Engine/Source/Runtime/AutomationWorker/Private/AutomationWorkerModule.cpp b/Engine/Source/Runtime/AutomationWorker/Private/AutomationWorkerModule.cpp
--- a/Engine/Source/Runtime/AutomationWorker/Private/AutomationWorkerModule.cpp
+++ b/Engine/Source/Runtime/AutomationWorker/Private/AutomationWorkerModule.cpp
@@ -167,7 +167,7 @@ void FAutomationWorkerModule::ReportNetworkCommandComplete()
 {
 	if (GIsAutomationTesting)
 	{
-		MessageEndpoint->Send(new FAutomationWorkerRequestNextNetworkCommand(ExecutionCount), TestRequesterAddress);
+		MessageEndpoint->Send(FMessageEndpoint::MakeMessage<FAutomationWorkerRequestNextNetworkCommand>(ExecutionCount), TestRequesterAddress);
 		if (StopTestEvent.IsBound())
 		{
 			// this is a local test; the message to continue will never arrive, so lets not wait for it
@@ -233,7 +233,7 @@ void FAutomationWorkerModule::ReportTestComplete()
 		else
 		{
 			// send the results to the controller
-			FAutomationWorkerRunTestsReply* Message = new FAutomationWorkerRunTestsReply();
+			FAutomationWorkerRunTestsReply* Message = FMessageEndpoint::MakeMessage<FAutomationWorkerRunTestsReply>();
 
 			Message->TestName = TestName;
 			Message->ExecutionCount = ExecutionCount;
@@ -274,7 +274,7 @@ void FAutomationWorkerModule::ReportTestComplete()
 
 void FAutomationWorkerModule::SendTests( const FMessageAddress& ControllerAddress )
 {
-	FAutomationWorkerRequestTestsReplyComplete* Reply = new FAutomationWorkerRequestTestsReplyComplete();
+	FAutomationWorkerRequestTestsReplyComplete* Reply = FMessageEndpoint::MakeMessage<FAutomationWorkerRequestTestsReplyComplete>();
 	for( int32 TestIndex = 0; TestIndex < TestInfo.Num(); TestIndex++ )
 	{
 		Reply->Tests.Emplace(FAutomationWorkerSingleTestReply(TestInfo[TestIndex]));
@@ -322,7 +322,7 @@ void FAutomationWorkerModule::HandleFindWorkersMessage(const FAutomationWorkerFi
 
 void FAutomationWorkerModule::SendWorkerFound(const FMessageAddress& ControllerAddress)
 {
-	FAutomationWorkerFindWorkersResponse* Response = new FAutomationWorkerFindWorkersResponse();
+	FAutomationWorkerFindWorkersResponse* Response = FMessageEndpoint::MakeMessage<FAutomationWorkerFindWorkersResponse>();
 
 	FString OSMajorVersionString, OSSubVersionString;
 	FPlatformMisc::GetOSVersions(OSMajorVersionString, OSSubVersionString);
@@ -363,7 +363,7 @@ void FAutomationWorkerModule::HandleNextNetworkCommandReplyMessage( const FAutom
 
 void FAutomationWorkerModule::HandlePingMessage( const FAutomationWorkerPing& Message, const TSharedRef<IMessageContext, ESPMode::ThreadSafe>& Context )
 {
-	MessageEndpoint->Send(new FAutomationWorkerPong(), Context->GetSender());
+	MessageEndpoint->Send(FMessageEndpoint::MakeMessage<FAutomationWorkerPong>(), Context->GetSender());
 }
 
 
@@ -456,7 +456,7 @@ void FAutomationWorkerModule::HandleScreenShotAndTraceCapturedWithName(const TAr
 	// Send the screen shot if we have a target
 	if (TestRequesterAddress.IsValid())
 	{
-		FAutomationWorkerScreenImage* Message = new FAutomationWorkerScreenImage();
+		FAutomationWorkerScreenImage* Message = FMessageEndpoint::MakeMessage<FAutomationWorkerScreenImage>();
 
 		Message->ScreenShotName = Data.ScreenshotName;
 		Message->ScreenImage = CompressedBitmap;
@@ -533,7 +533,7 @@ void FAutomationWorkerModule::HandleRunTestsMessage( const FAutomationWorkerRunT
 		UE_LOG(LogAutomationWorker, Warning, TEXT("%s"), *LogMessage);
 
 		// Let the sender know it won't happen
-		FAutomationWorkerRunTestsReply* OutMessage = new FAutomationWorkerRunTestsReply();
+		FAutomationWorkerRunTestsReply* OutMessage = FMessageEndpoint::MakeMessage<FAutomationWorkerRunTestsReply>();
 		OutMessage->TestName = Message.TestName;
 		OutMessage->ExecutionCount = Message.ExecutionCount;
 		OutMessage->Success = false;
@@ -593,7 +593,7 @@ void FAutomationWorkerModule::SendAnalyticsEvents(TArray<FString>& InAnalyticsIt
 
 void FAutomationWorkerModule::HandleTelemetryData(const FString& StorageName, const FString& InTestName, const TArray<FAutomationTelemetryData>& InItems)
 {
-	FAutomationWorkerTelemetryData* Message = new FAutomationWorkerTelemetryData();
+	FAutomationWorkerTelemetryData* Message = FMessageEndpoint::MakeMessage<FAutomationWorkerTelemetryData>();
 
 	Message->Storage = StorageName;
 	Message->Platform = FPlatformProperties::PlatformName();
diff --git a/Engine/Source/Runtime/Datasmith/DirectLink/Private/DirectLinkEndpoint.cpp b/Engine/Source/Runtime/Datasmith/DirectLink/Private/DirectLinkEndpoint.cpp
--- a/Engine/Source/Runtime/Datasmith/DirectLink/Private/DirectLinkEndpoint.cpp
+++ b/Engine/Source/Runtime/Datasmith/DirectLink/Private/DirectLinkEndpoint.cpp
@@ -1083,7 +1083,7 @@ void FInternalThreadState::Run()
 	}
 
 	UE_CLOG(SharedState.bDebugLog, LogDirectLinkNet, Verbose, TEXT("Endpoint '%s': Publishing FDirectLinkMsg_EndpointLifecycle Start"), *SharedState.NiceName);
-	MessageEndpoint->Publish(NewMessage<FDirectLinkMsg_EndpointLifecycle>(FDirectLinkMsg_EndpointLifecycle::ELifecycle::Start));
+	MessageEndpoint->Publish(FMessageEndpoint::MakeMessage<FDirectLinkMsg_EndpointLifecycle>(FDirectLinkMsg_EndpointLifecycle::ELifecycle::Start));
 
 	while (SharedState.bInnerThreadShouldRun)
 	{
@@ -1107,7 +1107,7 @@ void FInternalThreadState::Run()
 		if (Now_s - LastHeartbeatTime_s > gConfig.HeartbeatThreshold_s)
 		{
 			UE_CLOG(SharedState.bDebugLog, LogDirectLinkNet, Verbose, TEXT("Endpoint '%s': Publishing FDirectLinkMsg_EndpointLifecycle Heartbeat %f"), *SharedState.NiceName, Now_s);
-			MessageEndpoint->Publish(NewMessage<FDirectLinkMsg_EndpointLifecycle>(FDirectLinkMsg_EndpointLifecycle::ELifecycle::Heartbeat, ThisDescription.StateRevision));
+			MessageEndpoint->Publish(FMessageEndpoint::MakeMessage<FDirectLinkMsg_EndpointLifecycle>(FDirectLinkMsg_EndpointLifecycle::ELifecycle::Heartbeat, ThisDescription.StateRevision));
 			LastHeartbeatTime_s = Now_s;
 		}
 
@@ -1261,7 +1261,7 @@ void FInternalThreadState::Run()
 	}
 
 	UE_CLOG(SharedState.bDebugLog, LogDirectLinkNet, Display, TEXT("Endpoint '%s': Publishing FDirectLinkMsg_EndpointLifecycle Stop"), *SharedState.NiceName);
-	MessageEndpoint->Publish(NewMessage<FDirectLinkMsg_EndpointLifecycle>(FDirectLinkMsg_EndpointLifecycle::ELifecycle::Stop));
+	MessageEndpoint->Publish(FMessageEndpoint::MakeMessage<FDirectLinkMsg_EndpointLifecycle>(FDirectLinkMsg_EndpointLifecycle::ELifecycle::Stop));
 	FMessageEndpoint::SafeRelease(MessageEndpoint);
 }
 
diff --git a/Engine/Source/Runtime/Datasmith/DirectLink/Private/DirectLinkMessages.h b/Engine/Source/Runtime/Datasmith/DirectLink/Private/DirectLinkMessages.h
--- a/Engine/Source/Runtime/Datasmith/DirectLink/Private/DirectLinkMessages.h
+++ b/Engine/Source/Runtime/Datasmith/DirectLink/Private/DirectLinkMessages.h
@@ -12,23 +12,6 @@
 namespace DirectLink
 {
 class FScenePipeBase;
-
-
-// Helper function to allocate a UStruct message which memory will be owned and released
-// through FMessageContext.
-// The explicit FMemory::Malloc will match the FMemory::Free used there.
-// This is useful in scenarios where operator new isn't redirected to FMemory::Malloc.
-//
-// usage:
-//    FMyMessage* Message = NewMessage<FMyMessage>();
-//    Endpoint->Publish(Message); // The FMessageEndpoint now owns the Message allocation
-template<typename T, typename... ArgsType>
-T* NewMessage(ArgsType&&... Args)
-{
-	void* Memory = FMemory::Malloc(sizeof(T), alignof(T));
-	return new (Memory) T(Forward<ArgsType>(Args)...);
-}
-
 } // namespace DirectLink
 
 
diff --git a/Engine/Source/Runtime/Engine/Private/EngineService.cpp b/Engine/Source/Runtime/Engine/Private/EngineService.cpp
--- a/Engine/Source/Runtime/Engine/Private/EngineService.cpp
+++ b/Engine/Source/Runtime/Engine/Private/EngineService.cpp
@@ -41,7 +41,7 @@ void FEngineService::SendNotification( const TCHAR* NotificationText, const FMes
 {
 	if (MessageEndpoint.IsValid())
 	{
-		MessageEndpoint->Send(new FEngineServiceNotification(NotificationText, FPlatformTime::Seconds() - GStartTime), Recipient);
+		MessageEndpoint->Send(FMessageEndpoint::MakeMessage<FEngineServiceNotification>(NotificationText, FPlatformTime::Seconds() - GStartTime), Recipient);
 	}
 }
 
diff --git a/Engine/Source/Runtime/Launch/Private/IOS/GameLaunchDaemonMessageHandler.cpp b/Engine/Source/Runtime/Launch/Private/IOS/GameLaunchDaemonMessageHandler.cpp
--- a/Engine/Source/Runtime/Launch/Private/IOS/GameLaunchDaemonMessageHandler.cpp
+++ b/Engine/Source/Runtime/Launch/Private/IOS/GameLaunchDaemonMessageHandler.cpp
@@ -42,7 +42,7 @@ void FGameLaunchDaemonMessageHandler::HandlePingMessage(const FIOSLaunchDaemonPi
 	{
 		FMessageAddress MessageSender = Context->GetSender();
 
-		MessageEndpoint->Send(new FIOSLaunchDaemonPong(FString(FPlatformProperties::PlatformName()) + (TARGET_IPHONE_SIMULATOR ? FString(TEXT("Simulator:")) : FString(TEXT("@"))) + FString(FPlatformProcess::ComputerName())
+		MessageEndpoint->Send(FMessageEndpoint::MakeMessage<<FIOSLaunchDaemonPong>(FString(FPlatformProperties::PlatformName()) + (TARGET_IPHONE_SIMULATOR ? FString(TEXT("Simulator:")) : FString(TEXT("@"))) + FString(FPlatformProcess::ComputerName())
 			, FPlatformProcess::ComputerName()
 			, "Game_Running"
 			, [[UIDevice currentDevice] userInterfaceIdiom] == UIUserInterfaceIdiomPhone? "Phone" : "Tablet"
diff --git a/Engine/Source/Runtime/LiveLinkMessageBusFramework/Private/LiveLinkProvider.cpp b/Engine/Source/Runtime/LiveLinkMessageBusFramework/Private/LiveLinkProvider.cpp
--- a/Engine/Source/Runtime/LiveLinkMessageBusFramework/Private/LiveLinkProvider.cpp
+++ b/Engine/Source/Runtime/LiveLinkMessageBusFramework/Private/LiveLinkProvider.cpp
@@ -470,7 +470,7 @@ void FLiveLinkProvider::HandlePingMessage(const FLiveLinkPingMessage& Message, c
 		return;
 	}
 
-	MessageEndpoint->Send(new FLiveLinkPongMessage(ProviderName, MachineName, Message.PollRequest), Context->GetSender());
+	MessageEndpoint->Send(FMessageEndpoint::MakeMessage<FLiveLinkPongMessage>(ProviderName, MachineName, Message.PollRequest), Context->GetSender());
 }
 
 void FLiveLinkProvider::HandleConnectMessage(const FLiveLinkConnectMessage& Message, const TSharedRef<IMessageContext, ESPMode::ThreadSafe>& Context)
@@ -535,7 +535,7 @@ void FLiveLinkProvider::HandleHeartbeat(const FLiveLinkHeartbeatMessage& Message
 		TrackedAddress->LastHeartbeatTime = FPlatformTime::Seconds();
 
 		// Respond so editor gets heartbeat too
-		MessageEndpoint->Send(new FLiveLinkHeartbeatMessage(), Context->GetSender());
+		MessageEndpoint->Send(FMessageEndpoint::MakeMessage<FLiveLinkHeartbeatMessage>(), Context->GetSender());
 	}
 }
 
diff --git a/Engine/Source/Runtime/MessagingRpc/Private/MessageRpcClient.cpp b/Engine/Source/Runtime/MessagingRpc/Private/MessageRpcClient.cpp
--- a/Engine/Source/Runtime/MessagingRpc/Private/MessageRpcClient.cpp
+++ b/Engine/Source/Runtime/MessagingRpc/Private/MessageRpcClient.cpp
@@ -112,7 +112,7 @@ void FMessageRpcClient::CancelCall(const FGuid& CallId)
 	
 	if (Calls.RemoveAndCopyValue(CallId, Call))
 	{
-		MessageEndpoint->Send(new FMessageRpcCancel(CallId), ServerAddress);
+		MessageEndpoint->Send(FMessageEndpoint::MakeMessage<FMessageRpcCancel>(CallId), ServerAddress);
 	}
 }
 
diff --git a/Engine/Source/Runtime/MessagingRpc/Private/MessageRpcServer.cpp b/Engine/Source/Runtime/MessagingRpc/Private/MessageRpcServer.cpp
--- a/Engine/Source/Runtime/MessagingRpc/Private/MessageRpcServer.cpp
+++ b/Engine/Source/Runtime/MessagingRpc/Private/MessageRpcServer.cpp
@@ -115,7 +115,7 @@ void FMessageRpcServer::ProcessRequest(const TSharedRef<IMessageContext, ESPMode
 	else
 	{
 		// notify caller that call was not handled
-		MessageEndpoint->Send(new FMessageRpcUnhandled(Message->CallId), Context->GetSender());
+		MessageEndpoint->Send(FMessageEndpoint::MakeMessage<FMessageRpcUnhandled>(Message->CallId), Context->GetSender());
 	}
 }
 
diff --git a/Engine/Source/Runtime/Portal/Rpc/Private/PortalRpcLocator.cpp b/Engine/Source/Runtime/Portal/Rpc/Private/PortalRpcLocator.cpp
--- a/Engine/Source/Runtime/Portal/Rpc/Private/PortalRpcLocator.cpp
+++ b/Engine/Source/Runtime/Portal/Rpc/Private/PortalRpcLocator.cpp
@@ -68,9 +68,7 @@ private:
 		}
 
 		// @todo sarge: implement actual product GUID
-		// message is going to be deleted by FMemory::Free() (see FMessageContext destructor), so allocate it with Malloc
-		void* Memory = FMemory::Malloc(sizeof(FPortalRpcLocateServer));
-		MessageEndpoint->Publish(new(Memory) FPortalRpcLocateServer(FGuid(), EngineVersion, MacAddress, UserId), EMessageScope::Network);
+		MessageEndpoint->Publish(FMessageEndpoint::MakeMessage<FPortalRpcLocateServer>(FGuid(), EngineVersion, MacAddress, UserId), EMessageScope::Network);
 
 		return true;
 	}
diff --git a/Engine/Source/Runtime/Portal/Rpc/Private/PortalRpcServer.cpp b/Engine/Source/Runtime/Portal/Rpc/Private/PortalRpcServer.cpp
--- a/Engine/Source/Runtime/Portal/Rpc/Private/PortalRpcServer.cpp
+++ b/Engine/Source/Runtime/Portal/Rpc/Private/PortalRpcServer.cpp
@@ -14,7 +14,7 @@ public:
 
 	virtual void ConnectTo(const FMessageAddress& Address) const
 	{
-		MessageEndpoint->Send(new FPortalRpcServer(GetAddress().ToString()), Address);
+		MessageEndpoint->Send(FMessageEndpoint::MakeMessage<FPortalRpcServer>(GetAddress().ToString()), Address);
 	}
 
 	virtual IMessageRpcServer* GetMessageServer()
diff --git a/Engine/Source/Runtime/SessionServices/Private/SessionInstanceInfo.cpp b/Engine/Source/Runtime/SessionServices/Private/SessionInstanceInfo.cpp
--- a/Engine/Source/Runtime/SessionServices/Private/SessionInstanceInfo.cpp
+++ b/Engine/Source/Runtime/SessionServices/Private/SessionInstanceInfo.cpp
@@ -73,7 +73,7 @@ void FSessionInstanceInfo::ExecuteCommand(const FString& CommandString)
 {
 	if (MessageEndpoint.IsValid() && EngineAddress.IsValid())
 	{
-		MessageEndpoint->Send(new FEngineServiceExecuteCommand(CommandString, FPlatformProcess::UserName(false)), EngineAddress);
+		MessageEndpoint->Send(FMessageEndpoint::MakeMessage<FEngineServiceExecuteCommand>(CommandString, FPlatformProcess::UserName(false)), EngineAddress);
 	}
 }
 
@@ -82,7 +82,7 @@ void FSessionInstanceInfo::Terminate()
 {
 	if (MessageEndpoint.IsValid() && EngineAddress.IsValid())
 	{
-		MessageEndpoint->Send(new FEngineServiceTerminate(FPlatformProcess::UserName(false)), EngineAddress);
+		MessageEndpoint->Send(FMessageEndpoint::MakeMessage<FEngineServiceTerminate>(FPlatformProcess::UserName(false)), EngineAddress);
 	}
 }
 
