Deprecated nonsensical FDateTime + FDateTime operator that modified left-hand side

Updated the deprecated operator+ to not modify LHS
Didn't make it const since that causes ambiguity, e.g. FDateTime() + 0.0f

PR #8690: Fix for the FDateTime "+operator" (Contributed by [REDACTED_PERSON])

diff --git a/Engine/Source/Developer/ProfilerClient/Private/ProfilerClientManager.cpp b/Engine/Source/Developer/ProfilerClient/Private/ProfilerClientManager.cpp
--- a/Engine/Source/Developer/ProfilerClient/Private/ProfilerClientManager.cpp
+++ b/Engine/Source/Developer/ProfilerClient/Private/ProfilerClientManager.cpp
@@ -618,7 +618,7 @@ void FProfilerClientManager::HandleServicePingMessage(const FProfilerServicePing
 bool FProfilerClientManager::HandleTicker(float DeltaTime)
 {
 #if STATS
-	if (PendingInstances.Num() > 0 && FDateTime::Now() > LastPingTime + DeltaTime)
+	if (PendingInstances.Num() > 0 && FDateTime::Now() > LastPingTime + FTimespan(DeltaTime))
 	{
 		TArray<FGuid> Instances;
 		Instances.Append(PendingInstances);
diff --git a/Engine/Source/Runtime/Core/Public/Misc/DateTime.h b/Engine/Source/Runtime/Core/Public/Misc/DateTime.h
--- a/Engine/Source/Runtime/Core/Public/Misc/DateTime.h
+++ b/Engine/Source/Runtime/Core/Public/Misc/DateTime.h
@@ -127,17 +127,10 @@ public:
 		return *this;
 	}
 
-	/**
-	 * Adds the time from the given date to this date.
-	 *
-	 * @return This date.
-	 * @see FDateTime
-	 */
-	FDateTime& operator+(const FDateTime& Other)
+	UE_DEPRECATED(5.1, "Adding dates doesn't make sense. Please use FDateTime + FTimespan instead")
+	FDateTime operator+(const FDateTime& Other)
 	{
-		Ticks += Other.Ticks;
-
-		return *this;
+		return FDateTime(Ticks + Other.Ticks);
 	}
 
 	/**
diff --git a/Engine/Source/Runtime/Engine/Classes/Kismet/KismetMathLibrary.inl b/Engine/Source/Runtime/Engine/Classes/Kismet/KismetMathLibrary.inl
--- a/Engine/Source/Runtime/Engine/Classes/Kismet/KismetMathLibrary.inl
+++ b/Engine/Source/Runtime/Engine/Classes/Kismet/KismetMathLibrary.inl
@@ -2342,7 +2342,7 @@ FTimespan UKismetMathLibrary::Subtract_DateTimeDateTime(FDateTime A, FDateTime B
 KISMET_MATH_FORCEINLINE
 FDateTime UKismetMathLibrary::Add_DateTimeDateTime(FDateTime A, FDateTime B)
 {
-	return A + B;
+	return A + FTimespan(B.GetTicks());
 }
 
 KISMET_MATH_FORCEINLINE
