Legacy Session Frontend profiler should be less likely to lose the connection during live data preview.

Root cause was the Editor/UFE framerate dropping too low when there was a lot of stat data arriving. This meant the keep-alive Ping request from the game could not be acknowledged fast enough to avoid the 5s disconnection timeout.

 - decompress incoming stat messages via a task pipe to avoid saturating all cores.
 - only process a few stat messages per frame.
 - reduce the timeslice for stat aggregation during live data preview.

However, if bandwidth is saturated (typically by a very large amount of log spam being sent to the Session Frontend too) then the connection can still be lost.

Note: It's recommended to use Unreal Insights instead of the Session Profiler for UE5.

#rnx

diff --git a/Engine/Source/Developer/Profiler/Private/ProfilerSession.cpp b/Engine/Source/Developer/Profiler/Private/ProfilerSession.cpp
--- a/Engine/Source/Developer/Profiler/Private/ProfilerSession.cpp
+++ b/Engine/Source/Developer/Profiler/Private/ProfilerSession.cpp
@@ -381,8 +381,9 @@ bool FProfilerSession::HandleTicker( float DeltaTime )
 		NumFramesProcessedLastTime = 0;
 	}
 
-	// Limit processing to 50ms per frame.
-	const double TimeLimit = 250 / 1000.0;
+	// Limit processing per frame. For a live preview this must be severely limited so we can respond to the pings from the game in time
+	const double TimeLimitMS = ( bDataPreviewing && SessionType == EProfilerSessionTypes::Live ) ? 25.0 : 250.0;
+	const double TimeLimit = TimeLimitMS / 1000.0;
 	double Seconds = 0;
 
 	for( int32 It = 0; It < FrameToProcess.Num(); It++ )
diff --git a/Engine/Source/Developer/ProfilerClient/Private/ProfilerClientManager.cpp b/Engine/Source/Developer/ProfilerClient/Private/ProfilerClientManager.cpp
--- a/Engine/Source/Developer/ProfilerClient/Private/ProfilerClientManager.cpp
+++ b/Engine/Source/Developer/ProfilerClient/Private/ProfilerClientManager.cpp
@@ -5,6 +5,7 @@
 #include "MessageEndpointBuilder.h"
 #include "Misc/Paths.h"
 #include "Misc/SecureHash.h"
+#include "Async/Async.h"
 #include "Serialization/MemoryReader.h"
 #include "ProfilerServiceMessages.h"
 
@@ -22,6 +23,8 @@ DECLARE_CYCLE_STAT(TEXT("GenerateAccumulator"),STAT_PC_GenerateAccumulator,		STA
 DECLARE_CYCLE_STAT(TEXT("FindOrAddStat"),		STAT_PC_FindOrAddStat,				STATGROUP_Profiler);
 DECLARE_CYCLE_STAT(TEXT("FindOrAddThread"),	STAT_PC_FindOrAddThread,			STATGROUP_Profiler);
 
+UE::Tasks::FPipe FProfilerClientManager::AsyncTaskPipe{ TEXT("SessionProfilerClientPipe") };
+
 
 /* FProfilerClientManager structors
  *****************************************************************************/
@@ -47,9 +50,10 @@ FProfilerClientManager::FProfilerClientManager(const TSharedRef<IMessageBus, ESP
 	MessageDelegate = FTickerDelegate::CreateRaw(this, &FProfilerClientManager::HandleMessagesTicker);
 	LastPingTime = FDateTime::Now();
 	RetryTime = 5.f;
+	bIsLivePreview = false;
 
 	LoadConnection = nullptr;
-	MessageDelegateHandle = FTSTicker::GetCoreTicker().AddTicker(MessageDelegate, 0.1f);
+	MessageDelegateHandle = FTSTicker::GetCoreTicker().AddTicker(MessageDelegate);
 #endif
 }
 
@@ -181,6 +185,8 @@ void FProfilerClientManager::SetPreviewState(const bool bRequestedPreviewState,
 #if STATS
 	if (MessageEndpoint.IsValid() && ActiveSessionId.IsValid())
 	{
+		bIsLivePreview = bRequestedPreviewState;
+
 		if(!InstanceId.IsValid())
 		{
 			TArray<FMessageAddress> Instances;
@@ -641,61 +647,85 @@ bool FProfilerClientManager::HandleMessagesTicker(float DeltaTime)
     QUICK_SCOPE_CYCLE_COUNTER(STAT_FProfilerClientManager_HandleMessagesTicker);
 
 #if STATS
+	// MessageBus sends all data in out of order fashion.
+	// We buffer frame to make sure that all frames are received in the proper order.
+	const int32 NUM_BUFFERED_FRAMES = 15;
+
+	int32 NumConnectionsWithFrameData = 0;
 	for (auto It = Connections.CreateIterator(); It; ++It)
 	{
-		FServiceConnection& Connection = It.Value();
-
-		TArray<int64> Frames;
-		Connection.ReceivedData.GenerateKeyArray(Frames);
-		Frames.Sort();
+		if (It.Value().ReceivedData.Num() >= NUM_BUFFERED_FRAMES)
+		{
+			NumConnectionsWithFrameData++;
+		}
+	}
 
-		// MessageBus sends all data in out of order fashion.
-		// We buffer frame to make sure that all frames are received in the proper order.
-		const int32 NUM_BUFFERED_FRAMES = 15;
+	if (NumConnectionsWithFrameData)
+	{
+		// Limit all processing while doing a live preview - otherwise the ping from a live connection may not be processed in time & the game will disconnect us. 
+		// @todo all of this processing should move to a background thread or task, along with DecompressDataAndSendToGame 
+		const double TimeLimitSeconds = bIsLivePreview ? 0.2 : 0.8;
+		const double MaxDurationSeconds = TimeLimitSeconds / (double)NumConnectionsWithFrameData;
 
-		for(int32 Index = 0; Index < Frames.Num(); Index++)
+		for (auto It = Connections.CreateIterator(); It; ++It)
 		{
+			FServiceConnection& Connection = It.Value();
 			if (Connection.ReceivedData.Num() < NUM_BUFFERED_FRAMES)
 			{
 				break;
 			}
 
-			//FScopeLogTime SLT("HandleMessagesTicker");
-
-			const int64 FrameNum = Frames[Index];
-			const TArray<uint8>* const Data = Connection.ReceivedData.FindChecked(FrameNum);
-			FStatsReadStream& Stream = Connection.Stream;
+			TArray<int64> Frames;
+			Connection.ReceivedData.GenerateKeyArray(Frames);
+			Frames.Sort();
 
 
-			// Read all messages from the uncompressed buffer.
-			FMemoryReader MemoryReader(*Data, true);
-			while (MemoryReader.Tell() < MemoryReader.TotalSize())
+			uint64 StartTimeCycles = FPlatformTime::Cycles64();
+			for(int32 Index = 0; Index < Frames.Num(); Index++)
 			{
-				// Read the message.
-				FStatMessage Message(Stream.ReadMessage(MemoryReader));
-				new (Connection.PendingStatMessagesMessages)FStatMessage(Message);
-			}
+				const int64 FrameNum = Frames[Index];
+				const TArray<uint8>* const Data = Connection.ReceivedData.FindChecked(FrameNum);
+				FStatsReadStream& Stream = Connection.Stream;
 
-			// Adds a new from from the pending messages, the pending messages will be removed after the call.
-			Connection.CurrentThreadState.ProcessMetaDataAndLeaveDataOnly(Connection.PendingStatMessagesMessages);
-			Connection.CurrentThreadState.AddFrameFromCondensedMessages(Connection.PendingStatMessagesMessages);
 
-			UE_LOG(LogProfilerClient, VeryVerbose, TEXT("Frame=%i/%i, FNamesIndexMap=%i, CurrentMetadataSize=%i"), FrameNum, Frames.Num(), Connection.Stream.FNamesIndexMap.Num(), Connection.CurrentThreadState.ShortNameToLongName.Num());
+				// Read all messages from the uncompressed buffer.
+				FMemoryReader MemoryReader(*Data, true);
+				while (MemoryReader.Tell() < MemoryReader.TotalSize())
+				{
+					// Read the message.
+					FStatMessage Message(Stream.ReadMessage(MemoryReader));
+					new (Connection.PendingStatMessagesMessages)FStatMessage(Message);
+				}
 
-			// create an old format data frame from the data
-			Connection.GenerateProfilerDataFrame();
+				// Adds a new from from the pending messages, the pending messages will be removed after the call.
+				Connection.CurrentThreadState.ProcessMetaDataAndLeaveDataOnly(Connection.PendingStatMessagesMessages);
+				Connection.CurrentThreadState.AddFrameFromCondensedMessages(Connection.PendingStatMessagesMessages);
 
-			// Fire a meta data update message
-			if (Connection.CurrentData.MetaDataUpdated)
-			{
-				ProfilerMetaDataUpdatedDelegate.Broadcast(Connection.InstanceId, Connection.StatMetaData);
-			}
+				UE_LOG(LogProfilerClient, VeryVerbose, TEXT("Frame=%i/%i, FNamesIndexMap=%i, CurrentMetadataSize=%i"), FrameNum, Frames.Num(), Connection.Stream.FNamesIndexMap.Num(), Connection.CurrentThreadState.ShortNameToLongName.Num());
+
+				// create an old format data frame from the data
+				Connection.GenerateProfilerDataFrame();
 
-			// send the data out
-			ProfilerDataDelegate.Broadcast(Connection.InstanceId, Connection.CurrentData);
+				// Fire a meta data update message
+				if (Connection.CurrentData.MetaDataUpdated)
+				{
+					ProfilerMetaDataUpdatedDelegate.Broadcast(Connection.InstanceId, Connection.StatMetaData);
+				}
 
-			delete Data;
-			Connection.ReceivedData.Remove(FrameNum);
+				// send the data out
+				ProfilerDataDelegate.Broadcast(Connection.InstanceId, Connection.CurrentData);
+
+				delete Data;
+				Connection.ReceivedData.Remove(FrameNum);
+
+				// see if we need to yield
+				double DurationSeconds = (FPlatformTime::Cycles64() - StartTimeCycles) * FPlatformTime::GetSecondsPerCycle64();
+				if (DurationSeconds > MaxDurationSeconds)
+				{
+					UE_CLOG(Index < Frames.Num()-1, LogProfilerClient, Verbose, TEXT("Over time - %d/%d frames processed for connection %s"), Index, Frames.Num() - Index, *Connection.InstanceId.ToString() );
+					break;
+				}
+			}
 		}
 	}
 
@@ -746,12 +776,11 @@ void FProfilerClientManager::HandleProfilerServiceData2Message(const FProfilerSe
 		// Create a temporary profiler data and prepare all data.
 		FProfilerServiceData2* ToProcess = new FProfilerServiceData2(Message.InstanceId, Message.Frame, Message.HexData, Message.CompressedSize, Message.UncompressedSize);
 
-		// Decompression and decoding is done on the task graph.
-		FSimpleDelegateGraphTask::CreateAndDispatchWhenReady
-		(
-			FSimpleDelegateGraphTask::FDelegate::CreateRaw(this, &FProfilerClientManager::DecompressDataAndSendToGame, ToProcess), 
-			TStatId()
-		);
+		// Decompression uses a task pipe
+		AsyncTaskPipe.Launch(UE_SOURCE_LOCATION, [this,ToProcess]()
+		{
+			DecompressDataAndSendToGame(ToProcess);
+		});
 	}
 #endif
 }
diff --git a/Engine/Source/Developer/ProfilerClient/Private/ProfilerClientManager.h b/Engine/Source/Developer/ProfilerClient/Private/ProfilerClientManager.h
--- a/Engine/Source/Developer/ProfilerClient/Private/ProfilerClientManager.h
+++ b/Engine/Source/Developer/ProfilerClient/Private/ProfilerClientManager.h
@@ -12,6 +12,7 @@
 #include "IMessageBus.h"
 #include "MessageEndpoint.h"
 #include "Stats/StatsFile.h"
+#include "Tasks/Pipe.h"
 
 class FNewStatsReader;
 class FProfilerClientManager;
@@ -288,4 +289,10 @@ private:
 
 	/** Holds the last time a ping was made to instances */
 	FDateTime LastPingTime;
+
+	/** Async task pipe for processing the incoming profiler service data */
+	static UE::Tasks::FPipe AsyncTaskPipe;
+
+	/** Whether we are doing a data preview */
+	bool bIsLivePreview;
 };
